Don't let the user choose to use external actors or partioned world when creating a new world in the editor.
Instead, add a new template to create an empty partitioned world.

diff --git a/Engine/Source/Editor/EditorStyle/Private/SlateEditorStyle.cpp b/Engine/Source/Editor/EditorStyle/Private/SlateEditorStyle.cpp
--- a/Engine/Source/Editor/EditorStyle/Private/SlateEditorStyle.cpp
+++ b/Engine/Source/Editor/EditorStyle/Private/SlateEditorStyle.cpp
@@ -1750,6 +1750,7 @@ void FSlateEditorStyle::FStyle::SetupGeneralStyles()
 	{
 		Set( "NewLevelDialog.BlackBorder", new FSlateColorBrush( FColor(0, 0, 0, 100) ) );
 		Set( "NewLevelDialog.Blank", new IMAGE_BRUSH( "NewLevels/NewLevelBlank", FVector2D(256,256) ) );
+		Set( "NewLevelDialog.PartitionBlank", new IMAGE_BRUSH( "NewLevels/NewLevelPartitionBlank", FVector2D(256, 256) ) );
 		Set( "NewLevelDialog.Default", new IMAGE_BRUSH( "NewLevels/NewLevelDefault", FVector2D(256,256) ) );
 	}
 
diff --git a/Engine/Source/Editor/EditorStyle/Private/StarshipStyle.cpp b/Engine/Source/Editor/EditorStyle/Private/StarshipStyle.cpp
--- a/Engine/Source/Editor/EditorStyle/Private/StarshipStyle.cpp
+++ b/Engine/Source/Editor/EditorStyle/Private/StarshipStyle.cpp
@@ -1455,6 +1455,7 @@ void FStarshipEditorStyle::FStyle::SetupGeneralStyles()
 	{
 		Set( "NewLevelDialog.BlackBorder", new FSlateColorBrush( FColor(0, 0, 0, 100) ) );
 		Set( "NewLevelDialog.Blank", new IMAGE_BRUSH( "NewLevels/NewLevelBlank", FVector2D(256,256) ) );
+		Set( "NewLevelDialog.PartitionBlank", new IMAGE_BRUSH( "NewLevels/NewLevelPartitionBlank", FVector2D(256, 256) ) );
 		Set( "NewLevelDialog.Default", new IMAGE_BRUSH( "NewLevels/NewLevelDefault", FVector2D(256,256) ) );
 	}
 
diff --git a/Engine/Source/Editor/NewLevelDialog/Private/NewLevelDialogModule.cpp b/Engine/Source/Editor/NewLevelDialog/Private/NewLevelDialogModule.cpp
--- a/Engine/Source/Editor/NewLevelDialog/Private/NewLevelDialogModule.cpp
+++ b/Engine/Source/Editor/NewLevelDialog/Private/NewLevelDialogModule.cpp
@@ -87,6 +87,7 @@ private:
 	{
 		FTemplateMapInfo TemplateMapInfo;
 		bool bIsNewLevelItem;
+		bool bIsNewLevelItemPartitioned;
 	};
 
 public:
@@ -125,6 +126,11 @@ public:
 		NewItem->bIsNewLevelItem = true;
 		TemplateItemsList.Add(NewItem);
 
+		// Add an extra item for creating a new, blank partittioned level
+		TSharedPtr<FTemplateListItem> NewItemPartitioned = MakeShareable(new FTemplateListItem());
+		NewItemPartitioned->bIsNewLevelItemPartitioned = true;
+		TemplateItemsList.Add(NewItemPartitioned);
+
 		TSharedRef<SButton> CancelButton = SNew(SButton)
 			.ContentPadding(FMargin(10,3))
 			.Text(LOCTEXT("Cancel", "Cancel"))
@@ -140,34 +146,6 @@ public:
 			[
 				SNew(SVerticalBox)
 				+SVerticalBox::Slot()
-				.AutoHeight()
-				[
-					SNew(SVerticalBox)
-					+SVerticalBox::Slot()
-					[
-						SNew(SCheckBox)
-						.IsChecked_Lambda([=]() { return (bExternalActors || bPartitionedWorld) ? ECheckBoxState::Checked : ECheckBoxState::Unchecked; })
-						.IsEnabled_Lambda([=]() { return bExternalActorsEnabled && !bPartitionedWorld; })
-						.OnCheckStateChanged(FOnCheckStateChanged::CreateLambda([=](ECheckBoxState State) { bExternalActors = !bExternalActors; }))
-						[
-							SNew(STextBlock)
-							.Text(LOCTEXT("ExternalActors", "External Actors"))
-							.IsEnabled_Lambda([=]() { return bExternalActorsEnabled && !bPartitionedWorld; })
-						]
-					]
-					+SVerticalBox::Slot()
-					[
-						SNew(SCheckBox)
-						.IsChecked(ECheckBoxState::Unchecked)
-						.IsEnabled(bPartitionedWorldEnabled)
-						.OnCheckStateChanged(FOnCheckStateChanged::CreateLambda([=](ECheckBoxState State) { bPartitionedWorld = !bPartitionedWorld; }))
-						[
-							SNew(STextBlock)
-							.Text(LOCTEXT("PartitionedWorld", "Partitioned World"))
-						]
-					]
-				]
-				+SVerticalBox::Slot()
 				.FillHeight(1)
 				[
 					SNew(SScrollBox)
@@ -230,6 +208,12 @@ private:
 			Image = SNew(SImage).Image(FEditorStyle::GetBrush(TEXT("NewLevelDialog.Blank")));
 			Text = LOCTEXT("NewLevelItemLabel", "Empty Level");
 		}
+		else if (Template->bIsNewLevelItemPartitioned)
+		{
+			// New partitioned level item
+			Image = SNew(SImage).Image(FEditorStyle::GetBrush(TEXT("NewLevelDialog.PartitionBlank")));
+			Text = LOCTEXT("NewLevelItemPartitionedLabel", "Empty Partition Level");
+		}
 		else if (Template->TemplateMapInfo.ThumbnailTexture)
 		{
 			// Level with thumbnail
@@ -284,11 +268,12 @@ private:
 
 	FReply OnTemplateClicked(TSharedPtr<FTemplateListItem> Template)
 	{
-		if (!Template->bIsNewLevelItem)
+		if (!Template->bIsNewLevelItem && !Template->bIsNewLevelItemPartitioned)
 		{
 			OutTemplateMapPackageName = Template->TemplateMapInfo.Map;
 		}
 		bUserClickedOkay = true;
+		bPartitionedWorld = Template->bIsNewLevelItemPartitioned;
 		bExternalActors |= bPartitionedWorld;
 
 		ParentWindowPtr.Pin()->RequestDestroyWindow();
