Added a project settings to enable world partition.
Once enabled, creating new worlds will always make them partitioned and loading existing ones will ask for conversion.

diff --git a/Engine/Source/Editor/EditorStyle/Private/SlateEditorStyle.cpp b/Engine/Source/Editor/EditorStyle/Private/SlateEditorStyle.cpp
--- a/Engine/Source/Editor/EditorStyle/Private/SlateEditorStyle.cpp
+++ b/Engine/Source/Editor/EditorStyle/Private/SlateEditorStyle.cpp
@@ -1750,7 +1750,6 @@ void FSlateEditorStyle::FStyle::SetupGeneralStyles()
 	{
 		Set( "NewLevelDialog.BlackBorder", new FSlateColorBrush( FColor(0, 0, 0, 100) ) );
 		Set( "NewLevelDialog.Blank", new IMAGE_BRUSH( "NewLevels/NewLevelBlank", FVector2D(256,256) ) );
-		Set( "NewLevelDialog.PartitionBlank", new IMAGE_BRUSH( "NewLevels/NewLevelPartitionBlank", FVector2D(256, 256) ) );
 		Set( "NewLevelDialog.Default", new IMAGE_BRUSH( "NewLevels/NewLevelDefault", FVector2D(256,256) ) );
 	}
 
diff --git a/Engine/Source/Editor/EditorStyle/Private/StarshipStyle.cpp b/Engine/Source/Editor/EditorStyle/Private/StarshipStyle.cpp
--- a/Engine/Source/Editor/EditorStyle/Private/StarshipStyle.cpp
+++ b/Engine/Source/Editor/EditorStyle/Private/StarshipStyle.cpp
@@ -1464,7 +1464,6 @@ void FStarshipEditorStyle::FStyle::SetupGeneralStyles()
 	{
 		Set( "NewLevelDialog.BlackBorder", new FSlateColorBrush( FColor(0, 0, 0, 100) ) );
 		Set( "NewLevelDialog.Blank", new IMAGE_BRUSH( "NewLevels/NewLevelBlank", FVector2D(256,256) ) );
-		Set( "NewLevelDialog.PartitionBlank", new IMAGE_BRUSH( "NewLevels/NewLevelPartitionBlank", FVector2D(256, 256) ) );
 		Set( "NewLevelDialog.Default", new IMAGE_BRUSH( "NewLevels/NewLevelDefault", FVector2D(256,256) ) );
 	}
 
diff --git a/Engine/Source/Editor/NewLevelDialog/Private/NewLevelDialogModule.cpp b/Engine/Source/Editor/NewLevelDialog/Private/NewLevelDialogModule.cpp
--- a/Engine/Source/Editor/NewLevelDialog/Private/NewLevelDialogModule.cpp
+++ b/Engine/Source/Editor/NewLevelDialog/Private/NewLevelDialogModule.cpp
@@ -27,7 +27,7 @@
 #include "Editor/UnrealEdEngine.h"
 #include "Engine/Texture2D.h"
 #include "UnrealEdGlobals.h"
-
+#include "WorldPartition/IWorldPartitionEditorModule.h"
 
 #define LOCTEXT_NAMESPACE "NewLevelDialog"
 
@@ -87,7 +87,6 @@ private:
 	{
 		FTemplateMapInfo TemplateMapInfo;
 		bool bIsNewLevelItem;
-		bool bIsNewLevelItemPartitioned;
 	};
 
 public:
@@ -126,11 +125,6 @@ public:
 		NewItem->bIsNewLevelItem = true;
 		TemplateItemsList.Add(NewItem);
 
-		// Add an extra item for creating a new, blank partittioned level
-		TSharedPtr<FTemplateListItem> NewItemPartitioned = MakeShareable(new FTemplateListItem());
-		NewItemPartitioned->bIsNewLevelItemPartitioned = true;
-		TemplateItemsList.Add(NewItemPartitioned);
-
 		TSharedRef<SButton> CancelButton = SNew(SButton)
 			.ContentPadding(FMargin(10,3))
 			.Text(LOCTEXT("Cancel", "Cancel"))
@@ -208,12 +202,6 @@ private:
 			Image = SNew(SImage).Image(FEditorStyle::GetBrush(TEXT("NewLevelDialog.Blank")));
 			Text = LOCTEXT("NewLevelItemLabel", "Empty Level");
 		}
-		else if (Template->bIsNewLevelItemPartitioned)
-		{
-			// New partitioned level item
-			Image = SNew(SImage).Image(FEditorStyle::GetBrush(TEXT("NewLevelDialog.PartitionBlank")));
-			Text = LOCTEXT("NewLevelItemPartitionedLabel", "Empty Partition Level");
-		}
 		else if (Template->TemplateMapInfo.ThumbnailTexture)
 		{
 			// Level with thumbnail
@@ -268,12 +256,14 @@ private:
 
 	FReply OnTemplateClicked(TSharedPtr<FTemplateListItem> Template)
 	{
-		if (!Template->bIsNewLevelItem && !Template->bIsNewLevelItemPartitioned)
+		if (!Template->bIsNewLevelItem)
 		{
 			OutTemplateMapPackageName = Template->TemplateMapInfo.Map;
 		}
 		bUserClickedOkay = true;
-		bPartitionedWorld = Template->bIsNewLevelItemPartitioned;
+
+		IWorldPartitionEditorModule& WorldPartitionEditorModule = FModuleManager::LoadModuleChecked<IWorldPartitionEditorModule>("WorldPartitionEditor");
+		bPartitionedWorld = WorldPartitionEditorModule.IsWorldPartitionEnabled();
 		bExternalActors |= bPartitionedWorld;
 
 		ParentWindowPtr.Pin()->RequestDestroyWindow();
diff --git a/Engine/Source/Editor/WorldPartitionEditor/Private/WorldPartition/WorldPartitionEditorSettings.cpp b/Engine/Source/Editor/WorldPartitionEditor/Private/WorldPartition/WorldPartitionEditorSettings.cpp
new file mode 100644
--- /dev/null
+++ b/Engine/Source/Editor/WorldPartitionEditor/Private/WorldPartition/WorldPartitionEditorSettings.cpp
@@ -0,0 +1,23 @@
+// Copyright Epic Games, Inc. All Rights Reserved.
+
+#include "WorldPartition/WorldPartitionEditorSettings.h"
+
+#define LOCTEXT_NAMESPACE "WorldPartitionEditor"
+
+#if WITH_EDITOR
+bool UWorldPartitionEditorSettings::CanEditChange(const FProperty* InProperty) const
+{
+	if (InProperty)
+	{
+		const FName PropertyName = InProperty->GetFName();
+		if (PropertyName == GET_MEMBER_NAME_CHECKED(UWorldPartitionEditorSettings, bEnableWorldPartition))
+		{
+			return true;
+		}
+	}
+
+	return bEnableWorldPartition;
+}
+#endif
+
+#undef LOCTEXT_NAMESPACE
\ No newline at end of file
diff --git a/Engine/Source/Editor/WorldPartitionEditor/Private/WorldPartitionEditorModule.cpp b/Engine/Source/Editor/WorldPartitionEditor/Private/WorldPartitionEditorModule.cpp
--- a/Engine/Source/Editor/WorldPartitionEditor/Private/WorldPartitionEditorModule.cpp
+++ b/Engine/Source/Editor/WorldPartitionEditor/Private/WorldPartitionEditorModule.cpp
@@ -126,6 +126,11 @@ TSharedRef<SWidget> FWorldPartitionEditorModule::CreateWorldPartitionEditor()
 	return SNew(SWorldPartitionEditor).InWorld(EditorWorld);
 }
 
+bool FWorldPartitionEditorModule::IsWorldPartitionEnabled()
+{
+	return GetDefault<UWorldPartitionEditorSettings>()->bEnableWorldPartition;
+}
+
 bool FWorldPartitionEditorModule::ConvertMap(const FString& InLongPackageName)
 {
 	if (!GetDefault<UWorldPartitionEditorSettings>()->bEnableConversionPrompt || ULevel::GetIsLevelPartitionedFromPackage(FName(*InLongPackageName)))
diff --git a/Engine/Source/Editor/WorldPartitionEditor/Public/WorldPartition/WorldPartitionEditorSettings.h b/Engine/Source/Editor/WorldPartitionEditor/Public/WorldPartition/WorldPartitionEditorSettings.h
--- a/Engine/Source/Editor/WorldPartitionEditor/Public/WorldPartition/WorldPartitionEditorSettings.h
+++ b/Engine/Source/Editor/WorldPartitionEditor/Public/WorldPartition/WorldPartitionEditorSettings.h
@@ -15,6 +15,12 @@ class UWorldPartitionEditorSettings : public UDeveloperSettings
 public:
 	UWorldPartitionEditorSettings();
 
+#if WITH_EDITOR
+	virtual bool CanEditChange(const FProperty* InProperty) const override;
+#endif
+
+	UPROPERTY(config, EditAnywhere, Category = WorldPartition, meta = (ToolTip = "This option is used to enable World Partition support"))
+	bool bEnableWorldPartition = false;
 
 	UPROPERTY(config, EditAnywhere, Category = MapConversion, meta = (ToolTip = "This option when enabled will show a conversion prompt when opening non World Partition maps"))
 	bool bEnableConversionPrompt = false;
diff --git a/Engine/Source/Editor/WorldPartitionEditor/Public/WorldPartitionEditorModule.h b/Engine/Source/Editor/WorldPartitionEditor/Public/WorldPartitionEditorModule.h
--- a/Engine/Source/Editor/WorldPartitionEditor/Public/WorldPartitionEditorModule.h
+++ b/Engine/Source/Editor/WorldPartitionEditor/Public/WorldPartitionEditorModule.h
@@ -13,22 +13,27 @@ public:
 	/**
 	 * Called right after the module DLL has been loaded and the module object has been created
 	 */
-	virtual void StartupModule();
+	virtual void StartupModule() override;
 
 	/**
 	 * Called before the module is unloaded, right before the module object is destroyed.
 	 */
-	virtual void ShutdownModule();
+	virtual void ShutdownModule() override;
 	
 	/**
 	 * Creates a world partition widget
 	 */
 	virtual TSharedRef<class SWidget> CreateWorldPartitionEditor();
+
+	/**
+	 *
+	 */
+	virtual bool IsWorldPartitionEnabled() override;
 	
 	/**
 	 * 
 	 */
-	virtual bool ConvertMap(const FString& InLongPackageName);
+	virtual bool ConvertMap(const FString& InLongPackageName) override;
 
 private:
 	FDelegateHandle LevelEditorExtenderDelegateHandle;
diff --git a/Engine/Source/Runtime/Engine/Public/WorldPartition/IWorldPartitionEditorModule.h b/Engine/Source/Runtime/Engine/Public/WorldPartition/IWorldPartitionEditorModule.h
--- a/Engine/Source/Runtime/Engine/Public/WorldPartition/IWorldPartitionEditorModule.h
+++ b/Engine/Source/Runtime/Engine/Public/WorldPartition/IWorldPartitionEditorModule.h
@@ -12,5 +12,6 @@ class IWorldPartitionEditorModule : public IModuleInterface
 public:
 	virtual ~IWorldPartitionEditorModule() {}
 
+	virtual bool IsWorldPartitionEnabled() = 0;
 	virtual bool ConvertMap(const FString& InLongPackageName) = 0;
 };
\ No newline at end of file
