New Level dialog reskin - add Create button and convert to tile view (to match Project Browser style)

diff --git a/Engine/Source/Editor/EditorStyle/Private/StarshipStyle.cpp b/Engine/Source/Editor/EditorStyle/Private/StarshipStyle.cpp
--- a/Engine/Source/Editor/EditorStyle/Private/StarshipStyle.cpp
+++ b/Engine/Source/Editor/EditorStyle/Private/StarshipStyle.cpp
@@ -1435,7 +1435,6 @@ void FStarshipEditorStyle::FStyle::SetupGeneralStyles()
 	// New Level Dialog
 #if WITH_EDITOR || (IS_PROGRAM && WITH_UNREAL_DEVELOPER_TOOLS)
 	{
-		Set( "NewLevelDialog.BlackBorder", new FSlateColorBrush( FColor(0, 0, 0, 100) ) );
 		Set( "NewLevelDialog.Blank", new IMAGE_BRUSH( "NewLevels/NewLevelBlank", FVector2D(256,256) ) );
 		Set( "NewLevelDialog.Default", new IMAGE_BRUSH( "NewLevels/NewLevelDefault", FVector2D(256,256) ) );
 	}
diff --git a/Engine/Source/Editor/NewLevelDialog/Private/NewLevelDialogModule.cpp b/Engine/Source/Editor/NewLevelDialog/Private/NewLevelDialogModule.cpp
--- a/Engine/Source/Editor/NewLevelDialog/Private/NewLevelDialogModule.cpp
+++ b/Engine/Source/Editor/NewLevelDialog/Private/NewLevelDialogModule.cpp
@@ -2,78 +2,165 @@
 
 #include "NewLevelDialogModule.h"
 #include "Layout/Margin.h"
-#include "Textures/SlateShaderResource.h"
-#include "Rendering/RenderingCommon.h"
-#include "Rendering/DrawElements.h"
 #include "Widgets/DeclarativeSyntaxSupport.h"
 #include "Widgets/SCompoundWidget.h"
-#include "Widgets/SBoxPanel.h"
 #include "Widgets/SOverlay.h"
 #include "Widgets/SWindow.h"
-#include "Widgets/SLeafWidget.h"
 #include "Widgets/Text/STextBlock.h"
 #include "Modules/ModuleManager.h"
 #include "Misc/PackageName.h"
 #include "Framework/Application/SlateApplication.h"
 #include "Widgets/Layout/SBorder.h"
+#include "Widgets/Layout/SScrollBorder.h"
+#include "Widgets/Layout/SSeparator.h"
 #include "Widgets/Layout/SWrapBox.h"
 #include "Widgets/Images/SImage.h"
 #include "Widgets/Layout/SBox.h"
 #include "Widgets/Input/SButton.h"
-#include "Widgets/Layout/SScrollBox.h"
-#include "Widgets/Input/SCheckBox.h"
+#include "Widgets/Views/STileView.h"
 #include "EditorStyleSet.h"
-#include "RenderingThread.h"
 #include "Editor/UnrealEdEngine.h"
 #include "Engine/Texture2D.h"
 #include "UnrealEdGlobals.h"
 #include "WorldPartition/IWorldPartitionEditorModule.h"
+#include "Internationalization/BreakIterator.h"
+#include "Brushes/SlateImageBrush.h"
 
 #define LOCTEXT_NAMESPACE "NewLevelDialog"
 
+namespace NewLevelDialogDefs
+{
+	constexpr float TemplateTileHeight = 153;
+	constexpr float TemplateTileWidth = 102;
+	constexpr float DefaultWindowHeight = 418;
+	constexpr float DefaultWindowWidth = 527;
+	constexpr float LargeWindowHeight = 566;
+	constexpr float LargeWindowWidth = 1008;
+	constexpr float MinWindowHeight = 280;
+	constexpr float MinWindowWidth = 320;
+}
+
+struct FNewLevelTemplateItem
+{
+	FTemplateMapInfo TemplateMapInfo;
+	bool bIsNewLevelItem;
+	FText Name;
+	TUniquePtr<FSlateBrush> ThumbnailBrush;
+};
+
 /**
- * Widget class for rendering a UTexture2D in Slate
- * Work-in-progress idea that is defined here so that others don't use it yet
+ * Single thumbnail tile for a level template in the tile view of templates.
  */
-class STexture2DView : public SLeafWidget, public ISlateViewport, private TSlateTexture<FTexture2DRHIRef>
+class SNewLevelTemplateTile : public STableRow<TSharedPtr<FNewLevelTemplateItem>>
 {
 public:
-	SLATE_BEGIN_ARGS(STexture2DView) {}
+	SLATE_BEGIN_ARGS(SNewLevelTemplateTile) {}
+		SLATE_ARGUMENT(TSharedPtr<FNewLevelTemplateItem>, Item)
 	SLATE_END_ARGS()
 
-	void Construct( const FArguments& InArgs, UTexture2D* InTexture2D )
+	static TSharedRef<ITableRow> BuildTile(TSharedPtr<FNewLevelTemplateItem> Item, const TSharedRef<STableViewBase>& OwnerTable)
 	{
-		Size = FIntPoint(InTexture2D->GetSizeX(), InTexture2D->GetSizeY());
+		if (!ensure(Item.IsValid()))
+		{
+			return SNew(STableRow<TSharedPtr<FNewLevelTemplateItem>>, OwnerTable);
+		}
 
-		STexture2DView* TextureView = this;
-		ENQUEUE_RENDER_COMMAND(UpdateSTexture2DView)(
-			[TextureView, InTexture2D](FRHICommandListImmediate& RHICmdList)
-			{
-				TextureView->ShaderResource = InTexture2D->GetResource()->GetTexture2DRHI();
-			});
+		return SNew(SNewLevelTemplateTile, OwnerTable).Item(Item);
 	}
 
-	virtual int32 OnPaint( const FPaintArgs& Args, const FGeometry& AllottedGeometry, const FSlateRect& MyCullingRect, FSlateWindowElementList& OutDrawElements, int32 LayerId, const FWidgetStyle& InWidgetStyle, bool bParentEnabled ) const override
+	void Construct(const FArguments& InArgs, const TSharedRef<STableViewBase>& OwnerTable)
 	{
-		FSlateDrawElement::MakeViewport(
-			OutDrawElements, LayerId, AllottedGeometry.ToPaintGeometry(), SharedThis( this ), ESlateDrawEffect::NoBlending, InWidgetStyle.GetColorAndOpacityTint());
-		return LayerId;
+		check(InArgs._Item.IsValid());
+
+		STableRow::Construct(
+			STableRow::FArguments()
+			.Style(FAppStyle::Get(), "ProjectBrowser.TableRow")
+			.Padding(2.0f)
+			.Content()
+			[
+				SNew(SBorder)
+				.Padding(FMargin(0.0f, 0.0f, 5.0f, 5.0f))
+				.BorderImage(FAppStyle::Get().GetBrush("ProjectBrowser.ProjectTile.DropShadow"))
+				[
+					SNew(SOverlay)
+					+ SOverlay::Slot()
+					[
+						SNew(SVerticalBox)
+						// Thumbnail
+						+ SVerticalBox::Slot()
+						.AutoHeight()
+						.HAlign(HAlign_Center)
+						.VAlign(VAlign_Center)
+						[
+							SNew(SBox)
+							.WidthOverride(NewLevelDialogDefs::TemplateTileWidth)
+							.HeightOverride(NewLevelDialogDefs::TemplateTileWidth) // use width on purpose, this is a square
+							[
+								SNew(SBorder)
+								.Padding(0)
+								.BorderImage(FAppStyle::Get().GetBrush("ProjectBrowser.ProjectTile.ThumbnailAreaBackground"))
+								.HAlign(HAlign_Center)
+								.VAlign(VAlign_Center)
+								[
+									SNew(SImage)
+									.Image(InArgs._Item->ThumbnailBrush.Get())
+								]
+							]
+						]
+						// Name
+						+ SVerticalBox::Slot()
+						[
+							SNew(SBorder)
+							.Padding(FMargin(5.0f, 0))
+							.VAlign(VAlign_Top)
+							.Padding(FMargin(3.0f, 3.0f))
+							.BorderImage(FAppStyle::Get().GetBrush("ProjectBrowser.ProjectTile.NameAreaBackground"))
+							[
+								SNew(STextBlock)
+								.Font(FAppStyle::Get().GetFontStyle("ProjectBrowser.ProjectTile.Font"))
+								.WrapTextAt(NewLevelDialogDefs::TemplateTileWidth - 4.0f)
+								.LineBreakPolicy(FBreakIterator::CreateCamelCaseBreakIterator())
+								.Text(InArgs._Item->Name)
+								.ColorAndOpacity(FAppStyle::Get().GetSlateColor("Colors.Foreground"))
+							]
+						]
+					]
+					+ SOverlay::Slot()
+					[
+						SNew(SImage)
+						.Visibility(EVisibility::HitTestInvisible)
+						.Image(this, &SNewLevelTemplateTile::GetSelectionOutlineBrush)
+					]
+				]
+			],
+			OwnerTable
+		);
 	}
 
-	virtual FVector2D ComputeDesiredSize(float) const override
+private:
+	const FSlateBrush* GetSelectionOutlineBrush() const
 	{
-		return FVector2D(Size.X, Size.Y);
-	}
+		const bool bIsSelected = IsSelected();
+		const bool bIsTileHovered = IsHovered();
 
-	virtual FIntPoint GetSize() const override { return Size; }
-	virtual class FSlateShaderResource* GetViewportRenderTargetTexture() const override { return ShaderResource ? (FSlateShaderResource*)this : NULL; }
-	virtual bool RequiresVsync() const override { return false; }
+		if (bIsSelected && bIsTileHovered)
+		{
+			static const FName SelectedHover("ProjectBrowser.ProjectTile.SelectedHoverBorder");
+			return FAppStyle::Get().GetBrush(SelectedHover);
+		}
+		else if (bIsSelected)
+		{
+			static const FName Selected("ProjectBrowser.ProjectTile.SelectedBorder");
+			return FAppStyle::Get().GetBrush(Selected);
+		}
+		else if (bIsTileHovered)
+		{
+			static const FName Hovered("ProjectBrowser.ProjectTile.HoverBorder");
+			return FAppStyle::Get().GetBrush(Hovered);
+		}
 
-	// FSlateShaderResource implementation.
-	virtual uint32 GetWidth() const override { return Size.X; }
-	virtual uint32 GetHeight() const override { return Size.Y; }
-private:
-	FIntPoint Size;
+		return FStyleDefaults::GetNoBrush();
+	}
 };
 
 /**
@@ -82,185 +169,184 @@ private:
  */
 class SNewLevelDialog : public SCompoundWidget
 {
-private:
-	struct FTemplateListItem
-	{
-		FTemplateMapInfo TemplateMapInfo;
-		bool bIsNewLevelItem;
-	};
-
 public:
-	SLATE_BEGIN_ARGS(SNewLevelDialog)
-		: _ParentWindowClientSize(FVector2D::ZeroVector)
-		{}
+	SLATE_BEGIN_ARGS(SNewLevelDialog) {}
 		/** A pointer to the parent window */
 		SLATE_ATTRIBUTE(TSharedPtr<SWindow>, ParentWindow)
 
-
-		/** Initial size of the parent window */
-		SLATE_ARGUMENT( FVector2D, ParentWindowClientSize )
-
 		SLATE_ATTRIBUTE(TArray<FTemplateMapInfo>, Templates)
 	SLATE_END_ARGS()
 
 	void Construct(const FArguments& InArgs)
 	{
 		ParentWindowPtr = InArgs._ParentWindow.Get();
-		InitialWindowSize = InArgs._ParentWindowClientSize;
 
 		OutTemplateMapPackageName = TEXT("");
 		bUserClickedOkay = false;
 
-		const TArray<FTemplateMapInfo>& TemplateInfoList = InArgs._Templates.Get();
+		TemplateItemsList = MakeTemplateItems(InArgs._Templates.Get());
 
-		// Build a list of items - one for each template
-		for (int32 i=0; i < TemplateInfoList.Num(); i++)
-		{
-			TSharedPtr<FTemplateListItem> Item = MakeShareable(new FTemplateListItem());
-			Item->TemplateMapInfo = TemplateInfoList[i];
-			Item->bIsNewLevelItem = false;
-			TemplateItemsList.Add(Item);
-		}
-
-		// Add an extra item for creating a new, blank level
-		TSharedPtr<FTemplateListItem> NewItem = MakeShareable(new FTemplateListItem());
-		NewItem->bIsNewLevelItem = true;
-		TemplateItemsList.Add(NewItem);
+		TemplateListView = SNew(STileView<TSharedPtr<FNewLevelTemplateItem>>)
+			.ListItemsSource(&TemplateItemsList)
+			.SelectionMode(ESelectionMode::Single)
+			.ClearSelectionOnClick(false)
+			.ItemAlignment(EListItemAlignment::LeftAligned)
+			.OnGenerateTile_Static(&SNewLevelTemplateTile::BuildTile)
+			.OnMouseButtonDoubleClick(this, &SNewLevelDialog::HandleTemplateItemDoubleClick)
+			.ItemHeight(NewLevelDialogDefs::TemplateTileHeight + 9)
+			.ItemWidth(NewLevelDialogDefs::TemplateTileWidth + 9);
 
-		TSharedRef<SButton> CancelButton = SNew(SButton)
-			.ContentPadding(FMargin(10,3))
-			.Text(LOCTEXT("Cancel", "Cancel"))
-			.OnClicked(this, &SNewLevelDialog::OnCancelClicked);
+		TSharedPtr<SButton> CreateButton;
 
 		this->ChildSlot
 		[
 			SNew(SBorder)
-				.BorderImage(FEditorStyle::GetBrush("ToolPanel.GroupBorder"))
+			.BorderImage(FAppStyle::Get().GetBrush("Brushes.Panel"))
+			.Padding(FMargin(8.0f, 8.0f))
 			[
 				SNew(SVerticalBox)
+				// Top section with template thumbnails
 				+SVerticalBox::Slot()
-				.FillHeight(1)
 				[
-					SNew(SScrollBox)
-					+SScrollBox::Slot()
-					.Padding(15)
+					SNew(SScrollBorder, TemplateListView.ToSharedRef())
 					[
-						SAssignNew(TemplatesWrapBox, SWrapBox)
-						.PreferredSize(InitialWindowSize.X - 35.0)   // apparently no way to auto size the width of wrap boxes
+						TemplateListView.ToSharedRef()
 					]
 				]
 				+SVerticalBox::Slot()
 				.AutoHeight()
+				.Padding(-8.0f, 0.0f)
+				[
+					SNew(SSeparator)
+					.Orientation(EOrientation::Orient_Horizontal)
+					.Thickness(2.0f)
+				]
+				// Bottom section with dialog buttons
+				+SVerticalBox::Slot()
+				.AutoHeight()
 				.HAlign(HAlign_Right)
-				.Padding(6,2)
+				.Padding(8.0f, 16.0f, 8.0f, 8.0f)
 				[
-					CancelButton
+					SNew(SHorizontalBox)
+					+ SHorizontalBox::Slot()
+					.Padding(8.0f, 0.0f, 8.0f, 0.0f)
+					.VAlign(VAlign_Center)
+					.AutoWidth()
+					[
+						SAssignNew(CreateButton, SButton)
+						.HAlign(HAlign_Center)
+						.VAlign(VAlign_Center)
+						.ButtonStyle(FAppStyle::Get(), "PrimaryButton")
+						.TextStyle(FAppStyle::Get(), "DialogButtonText")
+						.Text(LOCTEXT("Create", "Create"))
+						.IsEnabled(this, &SNewLevelDialog::CanCreateLevel)
+						.OnClicked(this, &SNewLevelDialog::OnCreateClicked)
+					]
+					+ SHorizontalBox::Slot()
+					.Padding(8.0f, 0.0f, 0.0f, 0.0f)
+					.VAlign(VAlign_Center)
+					.AutoWidth()
+					[
+						SNew(SButton)
+						.HAlign(HAlign_Center)
+						.VAlign(VAlign_Center)
+						.TextStyle(FAppStyle::Get(), "DialogButtonText")
+						.Text(LOCTEXT("Cancel", "Cancel"))
+						.OnClicked(this, &SNewLevelDialog::OnCancelClicked)
+					]
 				]
 			]
 		];
 
-		// Give the cancel button initial focus so that the escape key can be checked for
-		ParentWindowPtr.Pin().Get()->SetWidgetToFocusOnActivate( CancelButton );
+		// Give the create button inital focus so that pressing enter will activate it
+		check(CreateButton != nullptr);
+		ParentWindowPtr.Pin().Get()->SetWidgetToFocusOnActivate( CreateButton );
 
-		// Insert items into slots in the wrap box, TemplatesWrapBox
-		AddItemsToWrapBox();
+		// Automatically select the first template item by default
+		if (!TemplateItemsList.IsEmpty())
+		{
+			TemplateListView->SetSelection(TemplateItemsList[0]);
+		}
 	}
 
-	/** A default window size for the dialog */
-	static const FVector2D DEFAULT_WINDOW_SIZE;
-
-	/** Level thumbnail image size in pixels */
-	static const float THUMBNAIL_SIZE;
-
 	FString GetChosenTemplate() const { return OutTemplateMapPackageName; }
 	bool IsTemplateChosen() const { return bUserClickedOkay; }
 
 private:
-	void AddItemsToWrapBox()
+	static TArray<TSharedPtr<FNewLevelTemplateItem>> MakeTemplateItems(const TArray<FTemplateMapInfo>& TemplateMapInfos)
 	{
-		for (int32 i = 0; i < TemplateItemsList.Num(); i++)
+		TArray<TSharedPtr<FNewLevelTemplateItem>> TemplateItems;
+
+		// Build a list of items - one for each template
+		for (const FTemplateMapInfo& TemplateMapInfo : TemplateMapInfos)
 		{
-			TSharedPtr<FTemplateListItem> Template = TemplateItemsList[i];
+			TSharedPtr<FNewLevelTemplateItem> Item = MakeShareable(new FNewLevelTemplateItem());
+			Item->TemplateMapInfo = TemplateMapInfo;
+			Item->bIsNewLevelItem = false;
 
-			TemplatesWrapBox->AddSlot()
-			[
-				GetWidgetForTemplate(TemplateItemsList[i])
-			];
+			if (const TObjectPtr<UTexture2D>& ThumbnailTexture = TemplateMapInfo.ThumbnailTexture)
+			{
+				// Level with thumbnail
+				Item->ThumbnailBrush = MakeUnique<FSlateImageBrush>(
+					ThumbnailTexture,
+					FVector2D(ThumbnailTexture->GetSizeX(), ThumbnailTexture->GetSizeY()));
+
+				if (Item->Name.IsEmpty())
+				{
+					Item->Name = FText::FromString(ThumbnailTexture->GetName().Replace(TEXT("_"), TEXT(" ")));
+				}
+			}
+			else
+			{
+				// Level with no thumbnail
+				Item->ThumbnailBrush = MakeUnique<FSlateBrush>(*FEditorStyle::GetBrush("NewLevelDialog.Default"));
+			}
+
+			check(Item->ThumbnailBrush);
+			Item->ThumbnailBrush->OutlineSettings.CornerRadii = FVector4(4, 4, 0, 0);
+			Item->ThumbnailBrush->OutlineSettings.RoundingType = ESlateBrushRoundingType::FixedRadius;
+			Item->ThumbnailBrush->DrawAs = ESlateBrushDrawType::RoundedBox;
+
+			TemplateItems.Add(Item);
 		}
+
+		// Add an extra item for creating a new, blank level
+		TSharedPtr<FNewLevelTemplateItem> NewItem = MakeShareable(new FNewLevelTemplateItem());
+		NewItem->bIsNewLevelItem = true;
+		NewItem->Name = LOCTEXT("NewLevelItemLabel", "Empty Level");
+		NewItem->ThumbnailBrush = MakeUnique<FSlateBrush>(*FEditorStyle::GetBrush("NewLevelDialog.Blank"));
+		NewItem->ThumbnailBrush->OutlineSettings.CornerRadii = FVector4(4, 4, 0, 0);
+		NewItem->ThumbnailBrush->OutlineSettings.RoundingType = ESlateBrushRoundingType::FixedRadius;
+		NewItem->ThumbnailBrush->DrawAs = ESlateBrushDrawType::RoundedBox;
+		TemplateItems.Add(NewItem);
+
+		return TemplateItems;
 	}
 
-	TSharedRef<SWidget> GetWidgetForTemplate(TSharedPtr<FTemplateListItem> Template)
+	bool CanCreateLevel() const
 	{
-		TSharedPtr<SWidget> Image;
-		FText Text = Template->TemplateMapInfo.DisplayName;
-		if (Template->bIsNewLevelItem)
+		if (!ensure(TemplateListView.IsValid()))
 		{
-			// New level item
-			Image = SNew(SImage).Image(FEditorStyle::GetBrush(TEXT("NewLevelDialog.Blank")));
-			Text = LOCTEXT("NewLevelItemLabel", "Empty Level");
+			return false;
 		}
-		else if (Template->TemplateMapInfo.ThumbnailTexture)
-		{
-			// Level with thumbnail
-			Image = SNew(STexture2DView, Template->TemplateMapInfo.ThumbnailTexture);
 
-			if (Text.IsEmpty())
-			{
-			Text = FText::FromString(Template->TemplateMapInfo.ThumbnailTexture->GetName().Replace(TEXT("_"), TEXT(" ")));
-		}
-		}
-		else
-		{
-			// Level with no thumbnail
-			Image = SNew(SImage).Image(FEditorStyle::GetBrush(TEXT("NewLevelDialog.Default")));
+		return TemplateListView->GetNumItemsSelected() == 1;
+	}
 
-			if (Text.IsEmpty())
-			{
-			Text = FText::FromString(FPackageName::GetShortName(Template->TemplateMapInfo.Map).Replace(TEXT("_"), TEXT(" ")));
-		}
+	FReply OnCreateClicked()
+	{
+		if (!ensure(TemplateListView.IsValid()))
+		{
+			return FReply::Handled();
 		}
 
-		Image->SetCursor(EMouseCursor::Hand);
-
-		return SNew(SBox)
-			.HeightOverride(THUMBNAIL_SIZE)
-			.WidthOverride(THUMBNAIL_SIZE)
-			.Padding(5)
-			[
-				SNew(SButton)
-				.ButtonStyle( FEditorStyle::Get(), "NoBorder" )
-				.VAlign(VAlign_Center)
-				.HAlign(HAlign_Center)
-				.OnClicked( this, &SNewLevelDialog::OnTemplateClicked, Template )
-				[
-					SNew(SBorder)
-					.BorderImage(FEditorStyle::GetBrush(TEXT("NewLevelDialog.BlackBorder")))
-					.ColorAndOpacity(this, &SNewLevelDialog::GetTemplateColor, Image )
-					.Padding(6)
-					[
-						SNew(SOverlay)
-						+SOverlay::Slot()
-						[
-							Image.ToSharedRef()
-						]
-						+SOverlay::Slot()
-						.VAlign(VAlign_Bottom)
-						.HAlign(HAlign_Right)
-						.Padding(FMargin(0, 0, 5, 5))
-						[
-							SNew(STextBlock)
-							.Visibility(EVisibility::HitTestInvisible)
-							.ShadowOffset(FVector2D(1.0f, 1.0f))
-							.ColorAndOpacity(FLinearColor(1,1,1))
-							.Text(Text)
-						]
-					]
-				]
-			];
-	}
+		const TArray<TSharedPtr<FNewLevelTemplateItem>> Items = TemplateListView->GetSelectedItems();
+		if (!ensure(Items.Num() == 1))
+		{
+			return FReply::Handled();
+		}
 
-	FReply OnTemplateClicked(TSharedPtr<FTemplateListItem> Template)
-	{
+		const TSharedPtr<FNewLevelTemplateItem> Template = Items[0];
 		if (!Template->bIsNewLevelItem)
 		{
 			OutTemplateMapPackageName = Template->TemplateMapInfo.Map;
@@ -283,18 +369,6 @@ private:
 		return FReply::Handled();
 	}
 
-	FLinearColor GetTemplateColor(TSharedPtr<SWidget> TemplateWidget) const
-	{
-		if (TemplateWidget->IsHovered())
-		{
-			return FLinearColor(1,1,1);
-		}
-		else
-		{
-			return FLinearColor(0.75,0.75,0.75);
-		}
-	}
-
 	virtual FReply OnKeyDown( const FGeometry& MyGeometry, const FKeyEvent& InKeyEvent ) override
 	{
 		if( InKeyEvent.GetKey() == EKeys::Escape )
@@ -305,23 +379,25 @@ private:
 		return SCompoundWidget::OnKeyDown( MyGeometry, InKeyEvent );
 	}
 
+	void HandleTemplateItemDoubleClick(TSharedPtr<FNewLevelTemplateItem>)
+	{
+		OnCreateClicked();
+	}
+
 	/** Pointer to the parent window, so we know to destroy it when done */
 	TWeakPtr<SWindow> ParentWindowPtr;
 
 	/** Initial size of the parent window */
 	FVector2D InitialWindowSize;
 
-	TArray<TSharedPtr<FTemplateListItem>> TemplateItemsList;
-	TSharedPtr<SWrapBox> TemplatesWrapBox;
+	TArray<TSharedPtr<FNewLevelTemplateItem>> TemplateItemsList;
+	TSharedPtr < STileView<TSharedPtr<FNewLevelTemplateItem>> > TemplateListView;
 	FString OutTemplateMapPackageName;
 	bool bUserClickedOkay;
 	bool bExternalActors;
 	bool bPartitionedWorld;
 };
 
-const FVector2D SNewLevelDialog::DEFAULT_WINDOW_SIZE = FVector2D(527, 418);
-const float SNewLevelDialog::THUMBNAIL_SIZE = 160.0f;
-
 IMPLEMENT_MODULE( FNewLevelDialogModule, NewLevelDialog );
 
 const FName FNewLevelDialogModule::NewLevelDialogAppIdentifier( TEXT( "NewLevelDialogApp" ) );
@@ -343,16 +419,18 @@ bool FNewLevelDialogModule::CreateAndShowNewLevelDialog( const TSharedPtr<const
 bool FNewLevelDialogModule::CreateAndShowTemplateDialog( const TSharedPtr<const SWidget> ParentWidget, const FText& Title, const TArray<FTemplateMapInfo>& Templates, FString& OutTemplateMapPackageName )
 {
 	// Open larger window if there are enough templates
-	FVector2D WindowClientSize(SNewLevelDialog::DEFAULT_WINDOW_SIZE);
+	FVector2D WindowClientSize(NewLevelDialogDefs::DefaultWindowWidth, NewLevelDialogDefs::DefaultWindowHeight);
 	if (Templates.Num() > 9)
 	{
-		WindowClientSize = FVector2D(1008, 566);
+		WindowClientSize = FVector2D(NewLevelDialogDefs::LargeWindowWidth, NewLevelDialogDefs::LargeWindowHeight);
 	}
 	
 	TSharedPtr<SWindow> NewLevelWindow =
 		SNew(SWindow)
 		.Title(Title)
 		.ClientSize(WindowClientSize)
+		.MinHeight(NewLevelDialogDefs::MinWindowHeight)
+		.MinWidth(NewLevelDialogDefs::MinWindowWidth)
 		.SizingRule( ESizingRule::UserSized )
 		.SupportsMinimize(false)
 		.SupportsMaximize(false);
@@ -360,7 +438,6 @@ bool FNewLevelDialogModule::CreateAndShowTemplateDialog( const TSharedPtr<const
 	TSharedRef<SNewLevelDialog> NewLevelDialog =
 		SNew(SNewLevelDialog)
 		.ParentWindow(NewLevelWindow)
-		.ParentWindowClientSize(WindowClientSize)
 		.Templates(Templates);
 
 	NewLevelWindow->SetContent(NewLevelDialog);
