World Partition:
- New Level Dialog: Empty Open World Template
- Remove IsWorldPartitionEnabled project setting

diff --git a/Engine/Source/Editor/EditorStyle/Private/SlateEditorStyle.cpp b/Engine/Source/Editor/EditorStyle/Private/SlateEditorStyle.cpp
--- a/Engine/Source/Editor/EditorStyle/Private/SlateEditorStyle.cpp
+++ b/Engine/Source/Editor/EditorStyle/Private/SlateEditorStyle.cpp
@@ -1743,6 +1743,7 @@ void FSlateEditorStyle::FStyle::SetupGeneralStyles()
 	{
 		Set( "NewLevelDialog.BlackBorder", new FSlateColorBrush( FColor(0, 0, 0, 100) ) );
 		Set( "NewLevelDialog.Blank", new IMAGE_BRUSH( "NewLevels/NewLevelBlank", FVector2D(256,256) ) );
+		Set( "NewLevelDialog.BlankWP", new IMAGE_BRUSH("NewLevels/NewLevelBlankWP", FVector2D(256, 256) ) );
 		Set( "NewLevelDialog.Default", new IMAGE_BRUSH( "NewLevels/NewLevelDefault", FVector2D(256,256) ) );
 	}
 
diff --git a/Engine/Source/Editor/EditorStyle/Private/StarshipStyle.cpp b/Engine/Source/Editor/EditorStyle/Private/StarshipStyle.cpp
--- a/Engine/Source/Editor/EditorStyle/Private/StarshipStyle.cpp
+++ b/Engine/Source/Editor/EditorStyle/Private/StarshipStyle.cpp
@@ -1442,6 +1442,7 @@ void FStarshipEditorStyle::FStyle::SetupGeneralStyles()
 #if WITH_EDITOR || (IS_PROGRAM && WITH_UNREAL_DEVELOPER_TOOLS)
 	{
 		Set( "NewLevelDialog.Blank", new IMAGE_BRUSH( "NewLevels/NewLevelBlank", FVector2D(256,256) ) );
+		Set( "NewLevelDialog.BlankWP", new IMAGE_BRUSH("NewLevels/NewLevelBlankWP", FVector2D(256, 256) ) );
 		Set( "NewLevelDialog.Default", new IMAGE_BRUSH( "NewLevels/NewLevelDefault", FVector2D(256,256) ) );
 	}
 
diff --git a/Engine/Source/Editor/GameProjectGeneration/Private/GameProjectUtils.cpp b/Engine/Source/Editor/GameProjectGeneration/Private/GameProjectUtils.cpp
--- a/Engine/Source/Editor/GameProjectGeneration/Private/GameProjectUtils.cpp
+++ b/Engine/Source/Editor/GameProjectGeneration/Private/GameProjectUtils.cpp
@@ -271,16 +271,10 @@ namespace
 	{
 		if (InProjectInfo.bIsBlankTemplate)
 		{
-			ConfigValues.Emplace(TEXT("DefaultEngine.ini"),
-				TEXT("/Script/WorldPartitionEditor.WorldPartitionEditorSettings"),
-				TEXT("bEnableWorldPartition"),
-				TEXT("False"),
-				true /* ShouldReplaceExistingValue */);
-
 			ConfigValues.Emplace(TEXT("DefaultEngine.ini"),
 				TEXT("/Script/WorldPartitionEditor.WorldPartitionEditorSettings"),
 				TEXT("bEnableConversionPrompt"),
-				TEXT("True"),
+				TEXT("False"),
 				true /* ShouldReplaceExistingValue */);
 
 			ConfigValues.Emplace(TEXT("DefaultEngine.ini"),
diff --git a/Engine/Source/Editor/LevelEditor/Private/LevelEditorActions.cpp b/Engine/Source/Editor/LevelEditor/Private/LevelEditorActions.cpp
--- a/Engine/Source/Editor/LevelEditor/Private/LevelEditorActions.cpp
+++ b/Engine/Source/Editor/LevelEditor/Private/LevelEditorActions.cpp
@@ -210,14 +210,17 @@ void FLevelEditorActionCallbacks::NewLevel()
 	IMainFrameModule& MainFrameModule = FModuleManager::GetModuleChecked<IMainFrameModule>("MainFrame");
 
 	FString TemplateMapPackageName;
+	bool bOutIsPartitionedWorld = false;
+	const bool bShowPartitionedTemplates = true;
+
 	FNewLevelDialogModule& NewLevelDialogModule = FModuleManager::LoadModuleChecked<FNewLevelDialogModule>("NewLevelDialog");
 	
-	if (NewLevelDialogModule.CreateAndShowNewLevelDialog(MainFrameModule.GetParentWindow(), TemplateMapPackageName))
+	if (NewLevelDialogModule.CreateAndShowNewLevelDialog(MainFrameModule.GetParentWindow(), TemplateMapPackageName, bShowPartitionedTemplates, bOutIsPartitionedWorld))
 	{
 		// The new map screen will return a blank TemplateName if the user has selected to begin a new blank map
 		if (TemplateMapPackageName.IsEmpty())
 		{
-			GEditor->CreateNewMapForEditing();
+			GEditor->CreateNewMapForEditing(/*bPromptUserToSave=*/true, bOutIsPartitionedWorld);
 		}
 		else
 		{
diff --git a/Engine/Source/Editor/LevelInstanceEditor/Private/LevelInstanceEditorModule.cpp b/Engine/Source/Editor/LevelInstanceEditor/Private/LevelInstanceEditorModule.cpp
--- a/Engine/Source/Editor/LevelInstanceEditor/Private/LevelInstanceEditorModule.cpp
+++ b/Engine/Source/Editor/LevelInstanceEditor/Private/LevelInstanceEditorModule.cpp
@@ -379,7 +379,9 @@ namespace LevelInstanceMenuUtils
 
 			FNewLevelDialogModule& NewLevelDialogModule = FModuleManager::LoadModuleChecked<FNewLevelDialogModule>("NewLevelDialog");
 			FString TemplateMapPackage;
-			if (!GetMutableDefault<ULevelInstanceEditorSettings>()->TemplateMapInfos.Num() || NewLevelDialogModule.CreateAndShowTemplateDialog(MainFrameModule.GetParentWindow(), LOCTEXT("LevelInstanceTemplateDialog", "Choose Level Instance Template..."), GetMutableDefault<ULevelInstanceEditorSettings>()->TemplateMapInfos, TemplateMapPackage))
+			bool bOutIsPartitionedWorld = false;
+			const bool bShowPartitionedTemplates = false;
+			if (!GetMutableDefault<ULevelInstanceEditorSettings>()->TemplateMapInfos.Num() || NewLevelDialogModule.CreateAndShowTemplateDialog(MainFrameModule.GetParentWindow(), LOCTEXT("LevelInstanceTemplateDialog", "Choose Level Instance Template..."), GetMutableDefault<ULevelInstanceEditorSettings>()->TemplateMapInfos, TemplateMapPackage, bShowPartitionedTemplates, bOutIsPartitionedWorld))
 			{
 				UPackage* TemplatePackage = !TemplateMapPackage.IsEmpty() ? LoadPackage(nullptr, *TemplateMapPackage, LOAD_None) : nullptr;
 				
diff --git a/Engine/Source/Editor/NewLevelDialog/Private/NewLevelDialogModule.cpp b/Engine/Source/Editor/NewLevelDialog/Private/NewLevelDialogModule.cpp
--- a/Engine/Source/Editor/NewLevelDialog/Private/NewLevelDialogModule.cpp
+++ b/Engine/Source/Editor/NewLevelDialog/Private/NewLevelDialogModule.cpp
@@ -22,10 +22,10 @@
 #include "Editor/UnrealEdEngine.h"
 #include "Engine/Texture2D.h"
 #include "UnrealEdGlobals.h"
-#include "WorldPartition/IWorldPartitionEditorModule.h"
 #include "Internationalization/BreakIterator.h"
 #include "Brushes/SlateImageBrush.h"
 #include "SPrimaryButton.h"
+#include "Engine/Level.h"
 
 #define LOCTEXT_NAMESPACE "NewLevelDialog"
 
@@ -44,9 +44,15 @@ namespace NewLevelDialogDefs
 struct FNewLevelTemplateItem
 {
 	FTemplateMapInfo TemplateMapInfo;
-	bool bIsNewLevelItem;
 	FText Name;
 	TUniquePtr<FSlateBrush> ThumbnailBrush;
+
+	enum NewLevelType
+	{
+		Empty,
+		EmptyWorldPartition,
+		Template
+	} Type;
 };
 
 /**
@@ -176,6 +182,7 @@ public:
 		SLATE_ATTRIBUTE(TSharedPtr<SWindow>, ParentWindow)
 
 		SLATE_ATTRIBUTE(TArray<FTemplateMapInfo>, Templates)
+		SLATE_ATTRIBUTE(bool, bShowPartitionedTemplates)
 	SLATE_END_ARGS()
 
 	void Construct(const FArguments& InArgs)
@@ -183,10 +190,11 @@ public:
 		ParentWindowPtr = InArgs._ParentWindow.Get();
 
 		OutTemplateMapPackageName = TEXT("");
+		bIsPartitionedWorld = false;
 		bUserClickedOkay = false;
 
-		TemplateItemsList = MakeTemplateItems(InArgs._Templates.Get());
-
+		TemplateItemsList = MakeTemplateItems(InArgs._bShowPartitionedTemplates.Get(), InArgs._Templates.Get());
+		
 		TemplateListView = SNew(STileView<TSharedPtr<FNewLevelTemplateItem>>)
 			.ListItemsSource(&TemplateItemsList)
 			.SelectionMode(ESelectionMode::Single)
@@ -267,19 +275,25 @@ public:
 	}
 
 	FString GetChosenTemplate() const { return OutTemplateMapPackageName; }
+	bool IsPartitionedWorld() const { return bIsPartitionedWorld; }
 	bool IsTemplateChosen() const { return bUserClickedOkay; }
 
 private:
-	static TArray<TSharedPtr<FNewLevelTemplateItem>> MakeTemplateItems(const TArray<FTemplateMapInfo>& TemplateMapInfos)
+	static TArray<TSharedPtr<FNewLevelTemplateItem>> MakeTemplateItems(bool bShowPartitionedTemplates, const TArray<FTemplateMapInfo>& TemplateMapInfos)
 	{
 		TArray<TSharedPtr<FNewLevelTemplateItem>> TemplateItems;
 
 		// Build a list of items - one for each template
 		for (const FTemplateMapInfo& TemplateMapInfo : TemplateMapInfos)
 		{
+			if (!bShowPartitionedTemplates && ULevel::GetIsLevelPartitionedFromPackage(FName(*TemplateMapInfo.Map)))
+			{
+				continue;
+			}
+
 			TSharedPtr<FNewLevelTemplateItem> Item = MakeShareable(new FNewLevelTemplateItem());
 			Item->TemplateMapInfo = TemplateMapInfo;
-			Item->bIsNewLevelItem = false;
+			Item->Type = FNewLevelTemplateItem::NewLevelType::Template;
 
 			if (const TObjectPtr<UTexture2D>& ThumbnailTexture = TemplateMapInfo.ThumbnailTexture)
 			{
@@ -309,7 +323,7 @@ private:
 
 		// Add an extra item for creating a new, blank level
 		TSharedPtr<FNewLevelTemplateItem> NewItem = MakeShareable(new FNewLevelTemplateItem());
-		NewItem->bIsNewLevelItem = true;
+		NewItem->Type = FNewLevelTemplateItem::NewLevelType::Empty;
 		NewItem->Name = LOCTEXT("NewLevelItemLabel", "Empty Level");
 		NewItem->ThumbnailBrush = MakeUnique<FSlateBrush>(*FEditorStyle::GetBrush("NewLevelDialog.Blank"));
 		NewItem->ThumbnailBrush->OutlineSettings.CornerRadii = FVector4(4, 4, 0, 0);
@@ -317,6 +331,19 @@ private:
 		NewItem->ThumbnailBrush->DrawAs = ESlateBrushDrawType::RoundedBox;
 		TemplateItems.Add(NewItem);
 
+		if (bShowPartitionedTemplates)
+		{
+			// Add an extra item for creating a new, blank level
+			TSharedPtr<FNewLevelTemplateItem> NewItemWP = MakeShareable(new FNewLevelTemplateItem());
+			NewItemWP->Type = FNewLevelTemplateItem::NewLevelType::EmptyWorldPartition;
+			NewItemWP->Name = LOCTEXT("NewWPLevelItemLabel", "Empty Open World");
+			NewItemWP->ThumbnailBrush = MakeUnique<FSlateBrush>(*FEditorStyle::GetBrush("NewLevelDialog.BlankWP"));
+			NewItemWP->ThumbnailBrush->OutlineSettings.CornerRadii = FVector4(4, 4, 0, 0);
+			NewItemWP->ThumbnailBrush->OutlineSettings.RoundingType = ESlateBrushRoundingType::FixedRadius;
+			NewItemWP->ThumbnailBrush->DrawAs = ESlateBrushDrawType::RoundedBox;
+			TemplateItems.Add(NewItemWP);
+		}
+
 		return TemplateItems;
 	}
 
@@ -344,16 +371,16 @@ private:
 		}
 
 		const TSharedPtr<FNewLevelTemplateItem> Template = Items[0];
-		if (!Template->bIsNewLevelItem)
+		if (Template->Type == FNewLevelTemplateItem::NewLevelType::Template)
 		{
 			OutTemplateMapPackageName = Template->TemplateMapInfo.Map;
 		}
+		else if (Template->Type == FNewLevelTemplateItem::NewLevelType::EmptyWorldPartition)
+		{
+			bIsPartitionedWorld = true;
+		}
 		bUserClickedOkay = true;
 
-		IWorldPartitionEditorModule& WorldPartitionEditorModule = FModuleManager::LoadModuleChecked<IWorldPartitionEditorModule>("WorldPartitionEditor");
-		bPartitionedWorld = WorldPartitionEditorModule.IsWorldPartitionEnabled();
-		bExternalActors |= bPartitionedWorld;
-
 		ParentWindowPtr.Pin()->RequestDestroyWindow();
 		return FReply::Handled();
 	}
@@ -391,8 +418,7 @@ private:
 	TSharedPtr < STileView<TSharedPtr<FNewLevelTemplateItem>> > TemplateListView;
 	FString OutTemplateMapPackageName;
 	bool bUserClickedOkay;
-	bool bExternalActors;
-	bool bPartitionedWorld;
+	bool bIsPartitionedWorld;
 };
 
 IMPLEMENT_MODULE( FNewLevelDialogModule, NewLevelDialog );
@@ -407,13 +433,13 @@ void FNewLevelDialogModule::ShutdownModule()
 {
 }
 
-bool FNewLevelDialogModule::CreateAndShowNewLevelDialog( const TSharedPtr<const SWidget> ParentWidget, FString& OutTemplateMapPackageName )
+bool FNewLevelDialogModule::CreateAndShowNewLevelDialog( const TSharedPtr<const SWidget> ParentWidget, FString& OutTemplateMapPackageName, bool bShowPartitionedTemplates, bool& bOutIsPartitionedWorld)
 {
 	TArray<FTemplateMapInfo> EmptyTemplates;
-	return CreateAndShowTemplateDialog(ParentWidget, LOCTEXT("WindowHeader", "New Level"), GUnrealEd ? GUnrealEd->GetTemplateMapInfos() : EmptyTemplates, OutTemplateMapPackageName);
+	return CreateAndShowTemplateDialog(ParentWidget, LOCTEXT("WindowHeader", "New Level"), GUnrealEd ? GUnrealEd->GetTemplateMapInfos() : EmptyTemplates, OutTemplateMapPackageName, bShowPartitionedTemplates, bOutIsPartitionedWorld);
 }
 
-bool FNewLevelDialogModule::CreateAndShowTemplateDialog( const TSharedPtr<const SWidget> ParentWidget, const FText& Title, const TArray<FTemplateMapInfo>& Templates, FString& OutTemplateMapPackageName )
+bool FNewLevelDialogModule::CreateAndShowTemplateDialog( const TSharedPtr<const SWidget> ParentWidget, const FText& Title, const TArray<FTemplateMapInfo>& Templates, FString& OutTemplateMapPackageName, bool bShowPartitionedTemplates, bool& bOutIsPartitionedWorld)
 {
 	// Open larger window if there are enough templates
 	FVector2D WindowClientSize(NewLevelDialogDefs::DefaultWindowWidth, NewLevelDialogDefs::DefaultWindowHeight);
@@ -435,13 +461,15 @@ bool FNewLevelDialogModule::CreateAndShowTemplateDialog( const TSharedPtr<const
 	TSharedRef<SNewLevelDialog> NewLevelDialog =
 		SNew(SNewLevelDialog)
 		.ParentWindow(NewLevelWindow)
-		.Templates(Templates);
+		.Templates(Templates)
+		.bShowPartitionedTemplates(bShowPartitionedTemplates);
 
 	NewLevelWindow->SetContent(NewLevelDialog);
 
 	FSlateApplication::Get().AddModalWindow(NewLevelWindow.ToSharedRef(), ParentWidget);
 
 	OutTemplateMapPackageName = NewLevelDialog->GetChosenTemplate();
+	bOutIsPartitionedWorld = NewLevelDialog->IsPartitionedWorld();
 	
 	return NewLevelDialog->IsTemplateChosen();
 }
diff --git a/Engine/Source/Editor/NewLevelDialog/Public/NewLevelDialogModule.h b/Engine/Source/Editor/NewLevelDialog/Public/NewLevelDialogModule.h
--- a/Engine/Source/Editor/NewLevelDialog/Public/NewLevelDialogModule.h
+++ b/Engine/Source/Editor/NewLevelDialog/Public/NewLevelDialogModule.h
@@ -28,9 +28,11 @@ public:
 	 * 
 	 * @param	ParentWidget - The parent widget for the modal window showing the dialog
 	 * @param	OutTemplateName	- (out) The package name of the template map selected by the user. Empty if blank map selected.
+	 * @param	bShowPartitionedTemplates - If true, we allow partitioned world templates to be shown in the dialog.
+	 * @param	bOutPartitionedWorld - (out) If true and OutTemplateName is empty, the empty map should be partitioned.
 	 * @return	true if the user selected a valid item, false if the user canceled
 	 */
-	virtual bool CreateAndShowNewLevelDialog( const TSharedPtr<const SWidget> ParentWidget, FString& OutTemplateMapPackageName );
+	virtual bool CreateAndShowNewLevelDialog( const TSharedPtr<const SWidget> ParentWidget, FString& OutTemplateMapPackageName, bool bShowPartitionedTemplates, bool& bOutPartitionedWorld );
 
 	/**
 	 * Creates and show a window with an SNewLevelDialog
@@ -39,9 +41,11 @@ public:
 	 * @param   Title - The dialog's title
 	 * @param   Templates - The list of template to be shown in the dialog
 	 * @param	OutTemplateName	- (out) The package name of the template map selected by the user. Empty if blank map selected.
+	 * @param	bShowPartitionedTemplates - If true, we allow partitioned world templates to be shown in the dialog.
+	 * @param	bOutPartitionedWorld - (out) If true and OutTemplateName is empty, the empty map should be partitioned.
 	 * @return	true if the user selected a valid item, false if the user canceled
 	 */
-	virtual bool CreateAndShowTemplateDialog(const TSharedPtr<const SWidget> ParentWidget, const FText& Title, const TArray<FTemplateMapInfo>& Templates, FString& OutTemplateMapPackageName );
+	virtual bool CreateAndShowTemplateDialog(const TSharedPtr<const SWidget> ParentWidget, const FText& Title, const TArray<FTemplateMapInfo>& Templates, FString& OutTemplateMapPackageName, bool bShowPartitionedTemplates, bool& bOutPartitionedWorld);
 
 	/** New Level Dialog app identifier string */
 	static const FName NewLevelDialogAppIdentifier;
diff --git a/Engine/Source/Editor/UnrealEd/Classes/Editor/EditorEngine.h b/Engine/Source/Editor/UnrealEd/Classes/Editor/EditorEngine.h
--- a/Engine/Source/Editor/UnrealEd/Classes/Editor/EditorEngine.h
+++ b/Engine/Source/Editor/UnrealEd/Classes/Editor/EditorEngine.h
@@ -1926,9 +1926,11 @@ public:
 	/**
 	 * Clears out the current map, if any, and creates a new blank map.
 	 *
+	 * @param   bIsPartitionedWorld	If true, new map is partitioned.
+	 * 
 	 * @return	Pointer to the map that was created.
 	 */
-	UWorld* NewMap();
+	UWorld* NewMap(bool bIsPartitionedWorld = false);
 
 	/**
 	 * Exports the current map to the specified filename.
@@ -2497,7 +2499,7 @@ public:
 	/**
 	 * Prompts the user to save the current map if necessary, then creates a new (blank) map.
 	 */
-	void CreateNewMapForEditing(bool bPromptUserToSave = true);
+	void CreateNewMapForEditing(bool bPromptUserToSave = true, bool bIsPartitionedWorld = false);
 
 	/**
 	 * If a PIE world exists, give the user the option to terminate it.
diff --git a/Engine/Source/Editor/UnrealEd/Private/EditorServer.cpp b/Engine/Source/Editor/UnrealEd/Private/EditorServer.cpp
--- a/Engine/Source/Editor/UnrealEd/Private/EditorServer.cpp
+++ b/Engine/Source/Editor/UnrealEd/Private/EditorServer.cpp
@@ -2230,7 +2230,7 @@ bool UEditorEngine::ShouldAbortBecauseOfUnsavedWorld()
 /**
  * Prompts the user to save the current map if necessary, then creates a new (blank) map.
  */
-void UEditorEngine::CreateNewMapForEditing(bool bPromptUserToSave)
+void UEditorEngine::CreateNewMapForEditing(bool bPromptUserToSave, bool bIsPartitionedWorld)
 {
 	// If a PIE world exists, warn the user that the PIE session will be terminated.
 	// Abort if the user refuses to terminate the PIE session.
@@ -2274,14 +2274,14 @@ void UEditorEngine::CreateNewMapForEditing(bool bPromptUserToSave)
 		GLevelEditorModeTools().DeactivateMode( FBuiltinEditorModes::EM_MeshPaint );
 	}
 
-	NewMap();
+	NewMap(bIsPartitionedWorld);
 
 	FEditorFileUtils::ResetLevelFilenames();
 }
 
 #define LOCTEXT_NAMESPACE "EditorEngine"
 
-UWorld* UEditorEngine::NewMap()
+UWorld* UEditorEngine::NewMap(bool bIsPartitionedWorld)
 {
 	// If we have a PIE session kill it before creating a new map
 	if (PlayWorld)
@@ -2306,6 +2306,7 @@ UWorld* UEditorEngine::NewMap()
 	// Create a new world
 	UWorldFactory* Factory = NewObject<UWorldFactory>();
 	Factory->WorldType = EWorldType::Editor;
+	Factory->bCreateWorldPartition = bIsPartitionedWorld;
 	Factory->bInformEngineOfWorld = true;
 	Factory->FeatureLevel = DefaultWorldFeatureLevel;
 	UPackage* Pkg = CreatePackage(nullptr);
diff --git a/Engine/Source/Editor/UnrealEd/Private/Factories/WorldFactory.cpp b/Engine/Source/Editor/UnrealEd/Private/Factories/WorldFactory.cpp
--- a/Engine/Source/Editor/UnrealEd/Private/Factories/WorldFactory.cpp
+++ b/Engine/Source/Editor/UnrealEd/Private/Factories/WorldFactory.cpp
@@ -7,7 +7,6 @@
 #include "ThumbnailRendering/WorldThumbnailInfo.h"
 #include "EditorClassUtils.h"
 #include "Modules/ModuleManager.h"
-#include "WorldPartition/IWorldPartitionEditorModule.h"
 
 #define LOCTEXT_NAMESPACE "WorldFactory"
 
@@ -18,7 +17,7 @@ UWorldFactory::UWorldFactory(const FObjectInitializer& ObjectInitializer)
 	SupportedClass = UWorld::StaticClass();
 	WorldType = EWorldType::Inactive;
 	bInformEngineOfWorld = false;
-	bCreateWorldPartition = true;
+	bCreateWorldPartition = false;
 	FeatureLevel = ERHIFeatureLevel::Num;
 }
 
@@ -32,16 +31,13 @@ UObject* UWorldFactory::FactoryCreateNew(UClass* Class, UObject* InParent, FName
 	// Create a new world.
 	const bool bAddToRoot = false;
 
-	IWorldPartitionEditorModule& WorldPartitionEditorModule = FModuleManager::LoadModuleChecked<IWorldPartitionEditorModule>("WorldPartitionEditor");
-	const bool bCreateWorldPartitionSetting = WorldPartitionEditorModule.IsWorldPartitionEnabled();
-		
 	// Those are the init values taken from the default in UWorld::CreateWorld + CreateWorldPartition.
 	UWorld::InitializationValues InitValues = UWorld::InitializationValues()
 		.ShouldSimulatePhysics(false)
 		.EnableTraceCollision(true)
 		.CreateNavigation(WorldType == EWorldType::Editor)
 		.CreateAISystem(WorldType == EWorldType::Editor)
-		.CreateWorldPartition(bCreateWorldPartitionSetting && bCreateWorldPartition);
+		.CreateWorldPartition(bCreateWorldPartition);
 
 	UWorld* NewWorld = UWorld::CreateWorld(WorldType, bInformEngineOfWorld, Name, Cast<UPackage>(InParent), bAddToRoot, FeatureLevel, &InitValues);
 	GEditor->InitBuilderBrush(NewWorld);
diff --git a/Engine/Source/Editor/WorldBrowser/Private/StreamingLevels/StreamingLevelCollectionModel.cpp b/Engine/Source/Editor/WorldBrowser/Private/StreamingLevels/StreamingLevelCollectionModel.cpp
--- a/Engine/Source/Editor/WorldBrowser/Private/StreamingLevels/StreamingLevelCollectionModel.cpp
+++ b/Engine/Source/Editor/WorldBrowser/Private/StreamingLevels/StreamingLevelCollectionModel.cpp
@@ -427,9 +427,11 @@ const FLevelModelList& FStreamingLevelCollectionModel::GetInvalidSelectedLevels(
 void FStreamingLevelCollectionModel::CreateNewLevel_Executed()
 {
 	FString TemplateMapPackageName;
+	bool bOutIsPartitionedWorld = false;
+	const bool bShowPartitionedTemplates = false;
 	FNewLevelDialogModule& NewLevelDialogModule = FModuleManager::LoadModuleChecked<FNewLevelDialogModule>("NewLevelDialog");
 	IMainFrameModule& MainFrameModule = FModuleManager::LoadModuleChecked<IMainFrameModule>(TEXT("MainFrame"));
-	if (NewLevelDialogModule.CreateAndShowNewLevelDialog(MainFrameModule.GetParentWindow(), TemplateMapPackageName))
+	if (NewLevelDialogModule.CreateAndShowNewLevelDialog(MainFrameModule.GetParentWindow(), TemplateMapPackageName, bShowPartitionedTemplates, bOutIsPartitionedWorld))
 	{
 		UPackage* TemplatePackage = TemplateMapPackageName.Len() ? LoadPackage(nullptr, *TemplateMapPackageName, LOAD_None) : nullptr;
 		UWorld* TemplateWorld = TemplatePackage ? UWorld::FindWorldInPackage(TemplatePackage) : nullptr;
diff --git a/Engine/Source/Editor/WorldPartitionEditor/Private/WorldPartition/WorldPartitionEditorSettings.cpp b/Engine/Source/Editor/WorldPartitionEditor/Private/WorldPartition/WorldPartitionEditorSettings.cpp
deleted file mode 100644
diff --git a/Engine/Source/Editor/WorldPartitionEditor/Private/WorldPartitionEditorModule.cpp b/Engine/Source/Editor/WorldPartitionEditor/Private/WorldPartitionEditorModule.cpp
--- a/Engine/Source/Editor/WorldPartitionEditor/Private/WorldPartitionEditorModule.cpp
+++ b/Engine/Source/Editor/WorldPartitionEditor/Private/WorldPartitionEditorModule.cpp
@@ -131,14 +131,9 @@ TSharedRef<SWidget> FWorldPartitionEditorModule::CreateWorldPartitionEditor()
 	return SNew(SWorldPartitionEditor).InWorld(EditorWorld);
 }
 
-bool FWorldPartitionEditorModule::IsWorldPartitionEnabled() const
-{
-	return GetDefault<UWorldPartitionEditorSettings>()->bEnableWorldPartition;
-}
-
 bool FWorldPartitionEditorModule::IsConversionPromptEnabled() const
 {
-	return IsWorldPartitionEnabled() && GetDefault<UWorldPartitionEditorSettings>()->bEnableConversionPrompt;
+	return GetDefault<UWorldPartitionEditorSettings>()->bEnableConversionPrompt;
 }
 
 void FWorldPartitionEditorModule::SetConversionPromptEnabled(bool bEnabled)
diff --git a/Engine/Source/Editor/WorldPartitionEditor/Public/WorldPartition/WorldPartitionEditorSettings.h b/Engine/Source/Editor/WorldPartitionEditor/Public/WorldPartition/WorldPartitionEditorSettings.h
--- a/Engine/Source/Editor/WorldPartitionEditor/Public/WorldPartition/WorldPartitionEditorSettings.h
+++ b/Engine/Source/Editor/WorldPartitionEditor/Public/WorldPartition/WorldPartitionEditorSettings.h
@@ -15,13 +15,6 @@ class UWorldPartitionEditorSettings : public UDeveloperSettings
 public:
 	UWorldPartitionEditorSettings();
 
-#if WITH_EDITOR
-	virtual bool CanEditChange(const FProperty* InProperty) const override;
-#endif
-
-	UPROPERTY(config, EditAnywhere, Category = WorldPartition, meta = (ToolTip = "This option is used to enable World Partition support"))
-	bool bEnableWorldPartition = false;
-
 	UPROPERTY(config, EditAnywhere, Category = MapConversion, meta = (ToolTip = "This option when enabled will show a conversion prompt when opening non World Partition maps"))
 	bool bEnableConversionPrompt = false;
 
diff --git a/Engine/Source/Editor/WorldPartitionEditor/Public/WorldPartitionEditorModule.h b/Engine/Source/Editor/WorldPartitionEditor/Public/WorldPartitionEditorModule.h
--- a/Engine/Source/Editor/WorldPartitionEditor/Public/WorldPartitionEditorModule.h
+++ b/Engine/Source/Editor/WorldPartitionEditor/Public/WorldPartitionEditorModule.h
@@ -25,11 +25,6 @@ public:
 	 */
 	virtual TSharedRef<class SWidget> CreateWorldPartitionEditor();
 
-	/**
-	 *
-	 */
-	virtual bool IsWorldPartitionEnabled() const override;
-
 	/**
 	 *
 	 */
diff --git a/Engine/Source/Runtime/Engine/Public/WorldPartition/IWorldPartitionEditorModule.h b/Engine/Source/Runtime/Engine/Public/WorldPartition/IWorldPartitionEditorModule.h
--- a/Engine/Source/Runtime/Engine/Public/WorldPartition/IWorldPartitionEditorModule.h
+++ b/Engine/Source/Runtime/Engine/Public/WorldPartition/IWorldPartitionEditorModule.h
@@ -12,7 +12,6 @@ class IWorldPartitionEditorModule : public IModuleInterface
 public:
 	virtual ~IWorldPartitionEditorModule() {}
 
-	virtual bool IsWorldPartitionEnabled() const = 0;
 	virtual bool IsConversionPromptEnabled() const = 0;
 	virtual void SetConversionPromptEnabled(bool bEnabled) = 0;
 	virtual float GetAutoCellLoadingMaxWorldSize() const = 0;
