Convert template map overrides into long package names so they load properly

diff --git a/Engine/Source/Editor/NewLevelDialog/Private/NewLevelDialogModule.cpp b/Engine/Source/Editor/NewLevelDialog/Private/NewLevelDialogModule.cpp
--- a/Engine/Source/Editor/NewLevelDialog/Private/NewLevelDialogModule.cpp
+++ b/Engine/Source/Editor/NewLevelDialog/Private/NewLevelDialogModule.cpp
@@ -471,7 +471,7 @@ private:
 		const TSharedPtr<FNewLevelTemplateItem> Template = Items[0];
 		if (Template->Type == FNewLevelTemplateItem::NewLevelType::Template)
 		{
-			OutTemplateMapPackageName = Template->TemplateMapInfo.Map.ToString();
+			OutTemplateMapPackageName = Template->TemplateMapInfo.Map.GetLongPackageName();
 		}
 		else if (Template->Type == FNewLevelTemplateItem::NewLevelType::EmptyWorldPartition)
 		{
