Fix empty map loading instead of picked map in new level dialog.

diff --git a/Engine/Source/Editor/NewLevelDialog/Private/NewLevelDialogModule.cpp b/Engine/Source/Editor/NewLevelDialog/Private/NewLevelDialogModule.cpp
--- a/Engine/Source/Editor/NewLevelDialog/Private/NewLevelDialogModule.cpp
+++ b/Engine/Source/Editor/NewLevelDialog/Private/NewLevelDialogModule.cpp
@@ -472,6 +472,11 @@ private:
 		if (Template->Type == FNewLevelTemplateItem::NewLevelType::Template)
 		{
 			OutTemplateMapPackageName = Template->TemplateMapInfo.Map.GetLongPackageName();
+			// For backwards compatibility, handle the case where the map name format does not include the object name directly
+			if (OutTemplateMapPackageName.IsEmpty())
+			{
+				OutTemplateMapPackageName = Template->TemplateMapInfo.Map.GetAssetPathString();
+			}
 		}
 		else if (Template->Type == FNewLevelTemplateItem::NewLevelType::EmptyWorldPartition)
 		{
