Update the default thumbnail for levels in the new level dialog

diff --git a/Engine/Source/Editor/EditorStyle/Private/StarshipStyle.cpp b/Engine/Source/Editor/EditorStyle/Private/StarshipStyle.cpp
--- a/Engine/Source/Editor/EditorStyle/Private/StarshipStyle.cpp
+++ b/Engine/Source/Editor/EditorStyle/Private/StarshipStyle.cpp
@@ -1486,7 +1486,6 @@ void FStarshipEditorStyle::FStyle::SetupGeneralStyles()
 	{
 		Set( "NewLevelDialog.Blank", new IMAGE_BRUSH( "NewLevels/NewLevelBlank", FVector2D(256,256) ) );
 		Set( "NewLevelDialog.BlankWP", new IMAGE_BRUSH("NewLevels/NewLevelBlankWP", FVector2D(256, 256) ) );
-		Set( "NewLevelDialog.Default", new IMAGE_BRUSH( "NewLevels/NewLevelDefault", FVector2D(256,256) ) );
 	}
 
 	// Build and Submit
diff --git a/Engine/Source/Editor/NewLevelDialog/Private/NewLevelDialogModule.cpp b/Engine/Source/Editor/NewLevelDialog/Private/NewLevelDialogModule.cpp
--- a/Engine/Source/Editor/NewLevelDialog/Private/NewLevelDialogModule.cpp
+++ b/Engine/Source/Editor/NewLevelDialog/Private/NewLevelDialogModule.cpp
@@ -49,6 +49,7 @@ struct FNewLevelTemplateItem
 	FText Name;
 	TUniquePtr<FSlateBrush> ThumbnailBrush;
 	UTexture2D* ThumbnailAsset = nullptr;
+	bool bDefaultThumbnail = false;
 
 	enum NewLevelType
 	{
@@ -107,10 +108,10 @@ public:
 							.HeightOverride(NewLevelDialogDefs::TemplateTileWidth) // use width on purpose, this is a square
 							[
 								SNew(SBorder)
-								.Padding(0)
+								.Padding(InArgs._Item->bDefaultThumbnail ? FMargin(12.0f) : FMargin(0))
 								.BorderImage(FAppStyle::Get().GetBrush("ProjectBrowser.ProjectTile.ThumbnailAreaBackground"))
-								.HAlign(HAlign_Center)
-								.VAlign(VAlign_Center)
+								.HAlign(InArgs._Item->bDefaultThumbnail ? HAlign_Fill : HAlign_Center)
+								.VAlign(InArgs._Item->bDefaultThumbnail ? VAlign_Fill : VAlign_Center)
 								[
 									SNew(SImage)
 									.Image(InArgs._Item->ThumbnailBrush.Get())
@@ -403,7 +404,8 @@ private:
 			else
 			{
 				// Level with no thumbnail
-				Item->ThumbnailBrush = MakeUnique<FSlateBrush>(*FEditorStyle::GetBrush("NewLevelDialog.Default"));
+				Item->ThumbnailBrush = MakeUnique<FSlateBrush>(*FAppStyle::Get().GetBrush("UnrealCircle.Thin"));
+				Item->bDefaultThumbnail = true;
 
 				if (Item->Name.IsEmpty())
 				{
