Update project default and new level default thumbnail look

diff --git a/Engine/Source/Editor/EditorStyle/Private/StarshipStyle.cpp b/Engine/Source/Editor/EditorStyle/Private/StarshipStyle.cpp
--- a/Engine/Source/Editor/EditorStyle/Private/StarshipStyle.cpp
+++ b/Engine/Source/Editor/EditorStyle/Private/StarshipStyle.cpp
@@ -473,6 +473,8 @@ void FStarshipEditorStyle::FStyle::SetupGeneralStyles()
 		Set("Icons.Toolbar.Import", new CORE_IMAGE_BRUSH_SVG( "Starship/Common/import", Icon20x20) );
 	}
 
+	Set("UnrealDefaultThumbnail", new IMAGE_BRUSH("Starship/Common/Unreal_DefaultThumbnail", FVector2D(256, 256)));
+
 	Set( "WarningStripe", new IMAGE_BRUSH( "Common/WarningStripe", FVector2D(20,6), FLinearColor::White, ESlateBrushTileType::Horizontal ) );
 	
 	Set("RoundedWarning", new FSlateRoundedBoxBrush(FStyleColors::Transparent, 4.0f, FStyleColors::Warning, 1.0f));
diff --git a/Engine/Source/Editor/GameProjectGeneration/Private/SProjectBrowser.cpp b/Engine/Source/Editor/GameProjectGeneration/Private/SProjectBrowser.cpp
--- a/Engine/Source/Editor/GameProjectGeneration/Private/SProjectBrowser.cpp
+++ b/Engine/Source/Editor/GameProjectGeneration/Private/SProjectBrowser.cpp
@@ -183,12 +183,12 @@ public:
 						[
 							SNew(SBorder)
 							.BorderImage(FAppStyle::Get().GetBrush("ProjectBrowser.ProjectTile.ThumbnailAreaBackground"))
-							.HAlign(ProjectItem->ProjectThumbnail ? HAlign_Center : HAlign_Fill)
-							.VAlign(ProjectItem->ProjectThumbnail ? VAlign_Center : VAlign_Fill)
-							.Padding(ProjectItem->ProjectThumbnail ? FMargin(0) : FMargin(12.0f) )
+							.HAlign(HAlign_Fill)
+							.VAlign(VAlign_Fill)
+							.Padding(FMargin(0))
 							[
 								SNew(SImage)
-								.Image(ProjectItem->ProjectThumbnail ? ProjectItem->ProjectThumbnail.Get() : FAppStyle::Get().GetBrush("UnrealCircle.Thin"))
+								.Image(ProjectItem->ProjectThumbnail ? ProjectItem->ProjectThumbnail.Get() : FAppStyle::Get().GetBrush("UnrealDefaultThumbnail"))
 								.ColorAndOpacity(FAppStyle::Get().GetSlateColor("Colors.Foreground"))
 							]
 						]
diff --git a/Engine/Source/Editor/GameProjectGeneration/Private/SProjectDialog.cpp b/Engine/Source/Editor/GameProjectGeneration/Private/SProjectDialog.cpp
--- a/Engine/Source/Editor/GameProjectGeneration/Private/SProjectDialog.cpp
+++ b/Engine/Source/Editor/GameProjectGeneration/Private/SProjectDialog.cpp
@@ -299,7 +299,7 @@ private:
 		{
 			return ItemPtr->Thumbnail.Get();
 		}
-		return FEditorStyle::GetBrush("UnrealCircle.Thin");
+		return FEditorStyle::GetBrush("UnrealDefaultThumbnail");
 	}
 	
 };
diff --git a/Engine/Source/Editor/NewLevelDialog/Private/NewLevelDialogModule.cpp b/Engine/Source/Editor/NewLevelDialog/Private/NewLevelDialogModule.cpp
--- a/Engine/Source/Editor/NewLevelDialog/Private/NewLevelDialogModule.cpp
+++ b/Engine/Source/Editor/NewLevelDialog/Private/NewLevelDialogModule.cpp
@@ -49,7 +49,6 @@ struct FNewLevelTemplateItem
 	FText Name;
 	TUniquePtr<FSlateBrush> ThumbnailBrush;
 	UTexture2D* ThumbnailAsset = nullptr;
-	bool bDefaultThumbnail = false;
 
 	enum NewLevelType
 	{
@@ -108,10 +107,10 @@ public:
 							.HeightOverride(NewLevelDialogDefs::TemplateTileWidth) // use width on purpose, this is a square
 							[
 								SNew(SBorder)
-								.Padding(InArgs._Item->bDefaultThumbnail ? FMargin(12.0f) : FMargin(0))
+								.Padding(FMargin(0))
 								.BorderImage(FAppStyle::Get().GetBrush("ProjectBrowser.ProjectTile.ThumbnailAreaBackground"))
-								.HAlign(InArgs._Item->bDefaultThumbnail ? HAlign_Fill : HAlign_Center)
-								.VAlign(InArgs._Item->bDefaultThumbnail ? VAlign_Fill : VAlign_Center)
+								.HAlign(HAlign_Fill)
+								.VAlign(VAlign_Fill)
 								[
 									SNew(SImage)
 									.Image(InArgs._Item->ThumbnailBrush.Get())
@@ -404,8 +403,7 @@ private:
 			else
 			{
 				// Level with no thumbnail
-				Item->ThumbnailBrush = MakeUnique<FSlateBrush>(*FAppStyle::Get().GetBrush("UnrealCircle.Thin"));
-				Item->bDefaultThumbnail = true;
+				Item->ThumbnailBrush = MakeUnique<FSlateBrush>(*FAppStyle::Get().GetBrush("UnrealDefaultThumbnail"));
 
 				if (Item->Name.IsEmpty())
 				{
