New Level Dialog: Sort Open World maps first, remove VR & TimeOfDay from BaseEngine.ini

#rnx

diff --git a/Engine/Config/BaseEngine.ini b/Engine/Config/BaseEngine.ini
--- a/Engine/Config/BaseEngine.ini
+++ b/Engine/Config/BaseEngine.ini
@@ -1514,10 +1514,8 @@ InEditorGameURLOptions=
 
 [/Script/UnrealEd.UnrealEdEngine]
 AutoSaveIndex=0
-+TemplateMapInfos=(Thumbnail="/Engine/Maps/Templates/Thumbnails/OpenWorld",Map="/Engine/Maps/Templates/OpenWorld")
-+TemplateMapInfos=(Thumbnail="/Engine/Maps/Templates/Thumbnails/Basic",Map="/Engine/Maps/Templates/Template_Default")
-+TemplateMapInfos=(Thumbnail="/Engine/Maps/Templates/Thumbnails/TimeOfDay",Map="/Engine/Maps/Templates/TimeOfDay_Default")
-+TemplateMapInfos=(Thumbnail="/Engine/Maps/Templates/Thumbnails/VR-Basic",Map="/Engine/Maps/Templates/VR-Basic")
++TemplateMapInfos=(Thumbnail="/Engine/Maps/Templates/Thumbnails/OpenWorld",Map="/Engine/Maps/Templates/OpenWorld",DisplayName="Open World",Category="OpenWorld")
++TemplateMapInfos=(Thumbnail="/Engine/Maps/Templates/Thumbnails/Basic",Map="/Engine/Maps/Templates/Template_Default",DisplayName="Basic")
 
 +PackagesToBeFullyLoadedAtStartup=/Engine/EditorMaterials/Cross
 +PackagesToBeFullyLoadedAtStartup=/Engine/EditorMaterials/Cross_Mat
diff --git a/Engine/Source/Editor/NewLevelDialog/Private/NewLevelDialogModule.cpp b/Engine/Source/Editor/NewLevelDialog/Private/NewLevelDialogModule.cpp
--- a/Engine/Source/Editor/NewLevelDialog/Private/NewLevelDialogModule.cpp
+++ b/Engine/Source/Editor/NewLevelDialog/Private/NewLevelDialogModule.cpp
@@ -47,6 +47,7 @@ struct FNewLevelTemplateItem
 {
 	FTemplateMapInfo TemplateMapInfo;
 	FText Name;
+	FString Category;
 	TUniquePtr<FSlateBrush> ThumbnailBrush;
 	UTexture2D* ThumbnailAsset = nullptr;
 
@@ -353,6 +354,7 @@ private:
 			Item->TemplateMapInfo = TemplateMapInfo;
 			Item->Type = FNewLevelTemplateItem::NewLevelType::Template;
 			Item->Name = TemplateMapInfo.DisplayName;
+			Item->Category = TemplateMapInfo.Category;
 
 			UTexture2D* ThumbnailTexture = nullptr;
 			if (TemplateMapInfo.Thumbnail.IsValid())
@@ -429,12 +431,15 @@ private:
 		NewItem->ThumbnailBrush->DrawAs = ESlateBrushDrawType::RoundedBox;
 		TemplateItems.Add(NewItem);
 
+		const FString& OpenWorldCategory = TEXT("OpenWorld");
+
 		if (bShowPartitionedTemplates)
 		{
 			// Add an extra item for creating a new, blank level
 			TSharedPtr<FNewLevelTemplateItem> NewItemWP = MakeShareable(new FNewLevelTemplateItem());
 			NewItemWP->Type = FNewLevelTemplateItem::NewLevelType::EmptyWorldPartition;
 			NewItemWP->Name = LOCTEXT("NewWPLevelItemLabel", "Empty Open World");
+			NewItemWP->Category = OpenWorldCategory;
 			NewItemWP->ThumbnailBrush = MakeUnique<FSlateBrush>(*FEditorStyle::GetBrush("NewLevelDialog.BlankWP"));
 			NewItemWP->ThumbnailBrush->OutlineSettings.CornerRadii = FVector4(4, 4, 0, 0);
 			NewItemWP->ThumbnailBrush->OutlineSettings.RoundingType = ESlateBrushRoundingType::FixedRadius;
@@ -442,6 +447,25 @@ private:
 			TemplateItems.Add(NewItemWP);
 		}
 
+		// OpenWorld category always first, If Categories are the same compares Type which will put Templates before Empty maps
+		TemplateItems.Sort([&OpenWorldCategory](const TSharedPtr<FNewLevelTemplateItem>& LHS, const TSharedPtr<FNewLevelTemplateItem>& RHS)
+		{
+			if (LHS->Category == RHS->Category)
+			{
+				return LHS->Type > RHS->Type;
+			}
+			else if (LHS->Category == OpenWorldCategory)
+			{
+				return true;
+			}
+			else if (RHS->Category == OpenWorldCategory)
+			{
+				return false;
+			}
+
+			return LHS->Category > RHS->Category;
+		});
+
 		return TemplateItems;
 	}
 
diff --git a/Engine/Source/Editor/UnrealEd/Classes/Editor/UnrealEdEngine.h b/Engine/Source/Editor/UnrealEd/Classes/Editor/UnrealEdEngine.h
--- a/Engine/Source/Editor/UnrealEd/Classes/Editor/UnrealEdEngine.h
+++ b/Engine/Source/Editor/UnrealEd/Classes/Editor/UnrealEdEngine.h
@@ -98,9 +98,13 @@ struct FTemplateMapInfo
 	FSoftObjectPath Map;
 
 	/** Optional display name override for this map template  */
-	UPROPERTY()
+	UPROPERTY(config)
 	FText DisplayName;
 
+	/* Optional category used for sorting */
+	UPROPERTY(config)
+	FString Category;
+
 	FTemplateMapInfo()
 	{
 	}
