Add support for multiple level templates using the same map package

diff --git a/Engine/Source/Editor/NewLevelDialog/Private/NewLevelDialogModule.cpp b/Engine/Source/Editor/NewLevelDialog/Private/NewLevelDialogModule.cpp
--- a/Engine/Source/Editor/NewLevelDialog/Private/NewLevelDialogModule.cpp
+++ b/Engine/Source/Editor/NewLevelDialog/Private/NewLevelDialogModule.cpp
@@ -230,10 +230,6 @@ public:
 	{
 		ParentWindowPtr = InArgs._ParentWindow.Get();
 
-		OutTemplateMapPackageName = TEXT("");
-		bIsPartitionedWorld = false;
-		bUserClickedOkay = false;
-
 		TemplateItemsList = MakeTemplateItems(InArgs._bShowPartitionedTemplates.Get(), InArgs._Templates.Get());
 		
 		TemplateListView = SNew(STileView<TSharedPtr<FNewLevelTemplateItem>>)
@@ -315,9 +311,25 @@ public:
 		}
 	}
 
-	FString GetChosenTemplate() const { return OutTemplateMapPackageName; }
-	bool IsPartitionedWorld() const { return bIsPartitionedWorld; }
-	bool IsTemplateChosen() const { return bUserClickedOkay; }
+	FString GetChosenTemplate() const
+	{
+		FString Result;
+		if (OutTemplateItem.IsValid())
+		{
+			Result = OutTemplateItem->TemplateMapInfo.Map.GetLongPackageName();
+			// For backwards compatibility, handle the case where the map name format does not include the object name directly
+			if (Result.IsEmpty())
+			{
+				Result = OutTemplateItem->TemplateMapInfo.Map.GetAssetPathString();
+			}
+		}
+
+		return Result;
+	}
+
+	FTemplateMapInfo GetChosenTemplateMapInfo() const { return OutTemplateItem.IsValid() ? OutTemplateItem->TemplateMapInfo : FTemplateMapInfo(); }
+	bool IsPartitionedWorld() const { return OutTemplateItem.IsValid() && OutTemplateItem->Type == FNewLevelTemplateItem::NewLevelType::EmptyWorldPartition; }
+	bool IsTemplateChosen() const { return OutTemplateItem.IsValid(); }
 
 	//~ Begin FGCObject Interface.
 	virtual void AddReferencedObjects(FReferenceCollector& Collector)  override
@@ -492,21 +504,7 @@ private:
 			return FReply::Handled();
 		}
 
-		const TSharedPtr<FNewLevelTemplateItem> Template = Items[0];
-		if (Template->Type == FNewLevelTemplateItem::NewLevelType::Template)
-		{
-			OutTemplateMapPackageName = Template->TemplateMapInfo.Map.GetLongPackageName();
-			// For backwards compatibility, handle the case where the map name format does not include the object name directly
-			if (OutTemplateMapPackageName.IsEmpty())
-			{
-				OutTemplateMapPackageName = Template->TemplateMapInfo.Map.GetAssetPathString();
-			}
-		}
-		else if (Template->Type == FNewLevelTemplateItem::NewLevelType::EmptyWorldPartition)
-		{
-			bIsPartitionedWorld = true;
-		}
-		bUserClickedOkay = true;
+		OutTemplateItem = Items[0];
 
 		ParentWindowPtr.Pin()->RequestDestroyWindow();
 		return FReply::Handled();
@@ -514,8 +512,6 @@ private:
 
 	FReply OnCancelClicked()
 	{
-		bUserClickedOkay = false;
-
 		ParentWindowPtr.Pin()->RequestDestroyWindow();
 		return FReply::Handled();
 	}
@@ -543,9 +539,7 @@ private:
 
 	TArray<TSharedPtr<FNewLevelTemplateItem>> TemplateItemsList;
 	TSharedPtr < STileView<TSharedPtr<FNewLevelTemplateItem>> > TemplateListView;
-	FString OutTemplateMapPackageName;
-	bool bUserClickedOkay;
-	bool bIsPartitionedWorld;
+	TSharedPtr<FNewLevelTemplateItem> OutTemplateItem;
 };
 
 IMPLEMENT_MODULE( FNewLevelDialogModule, NewLevelDialog );
@@ -560,13 +554,13 @@ void FNewLevelDialogModule::ShutdownModule()
 {
 }
 
-bool FNewLevelDialogModule::CreateAndShowNewLevelDialog( const TSharedPtr<const SWidget> ParentWidget, FString& OutTemplateMapPackageName, bool bShowPartitionedTemplates, bool& bOutIsPartitionedWorld)
+bool FNewLevelDialogModule::CreateAndShowNewLevelDialog( const TSharedPtr<const SWidget> ParentWidget, FString& OutTemplateMapPackageName, bool bShowPartitionedTemplates, bool& bOutIsPartitionedWorld, FTemplateMapInfo* OutTemplateMapInfo)
 {
 	TArray<FTemplateMapInfo> EmptyTemplates;
-	return CreateAndShowTemplateDialog(ParentWidget, LOCTEXT("WindowHeader", "New Level"), GUnrealEd ? GUnrealEd->GetTemplateMapInfos() : EmptyTemplates, OutTemplateMapPackageName, bShowPartitionedTemplates, bOutIsPartitionedWorld);
+	return CreateAndShowTemplateDialog(ParentWidget, LOCTEXT("WindowHeader", "New Level"), GUnrealEd ? GUnrealEd->GetTemplateMapInfos() : EmptyTemplates, OutTemplateMapPackageName, bShowPartitionedTemplates, bOutIsPartitionedWorld, OutTemplateMapInfo);
 }
 
-bool FNewLevelDialogModule::CreateAndShowTemplateDialog( const TSharedPtr<const SWidget> ParentWidget, const FText& Title, const TArray<FTemplateMapInfo>& Templates, FString& OutTemplateMapPackageName, bool bShowPartitionedTemplates, bool& bOutIsPartitionedWorld)
+bool FNewLevelDialogModule::CreateAndShowTemplateDialog( const TSharedPtr<const SWidget> ParentWidget, const FText& Title, const TArray<FTemplateMapInfo>& Templates, FString& OutTemplateMapPackageName, bool bShowPartitionedTemplates, bool& bOutIsPartitionedWorld, FTemplateMapInfo* OutTemplateMapInfo)
 {
 	// Open larger window if there are enough templates
 	FVector2D WindowClientSize(NewLevelDialogDefs::DefaultWindowWidth, NewLevelDialogDefs::DefaultWindowHeight);
@@ -595,6 +589,11 @@ bool FNewLevelDialogModule::CreateAndShowTemplateDialog( const TSharedPtr<const
 
 	FSlateApplication::Get().AddModalWindow(NewLevelWindow.ToSharedRef(), ParentWidget);
 
+	if (OutTemplateMapInfo)
+	{
+		*OutTemplateMapInfo = NewLevelDialog->GetChosenTemplateMapInfo();
+	}
+
 	OutTemplateMapPackageName = NewLevelDialog->GetChosenTemplate();
 	bOutIsPartitionedWorld = NewLevelDialog->IsPartitionedWorld();
 	
diff --git a/Engine/Source/Editor/NewLevelDialog/Public/NewLevelDialogModule.h b/Engine/Source/Editor/NewLevelDialog/Public/NewLevelDialogModule.h
--- a/Engine/Source/Editor/NewLevelDialog/Public/NewLevelDialogModule.h
+++ b/Engine/Source/Editor/NewLevelDialog/Public/NewLevelDialogModule.h
@@ -7,6 +7,8 @@
 #include "Modules/ModuleInterface.h"
 #include "Editor/UnrealEdEngine.h"
 
+struct FTemplateMapInfo;
+
 /**
  * New Level Dialog module
  */
@@ -30,9 +32,10 @@ public:
 	 * @param	OutTemplateName	- (out) The package name of the template map selected by the user. Empty if blank map selected.
 	 * @param	bShowPartitionedTemplates - If true, we allow partitioned world templates to be shown in the dialog.
 	 * @param	bOutPartitionedWorld - (out) If true and OutTemplateName is empty, the empty map should be partitioned.
+	 * @param	OutTemplateMapInfo - (out) If not null, will be set to the template map selected by the user.
 	 * @return	true if the user selected a valid item, false if the user canceled
 	 */
-	virtual bool CreateAndShowNewLevelDialog( const TSharedPtr<const SWidget> ParentWidget, FString& OutTemplateMapPackageName, bool bShowPartitionedTemplates, bool& bOutPartitionedWorld );
+	virtual bool CreateAndShowNewLevelDialog( const TSharedPtr<const SWidget> ParentWidget, FString& OutTemplateMapPackageName, bool bShowPartitionedTemplates, bool& bOutPartitionedWorld, FTemplateMapInfo* OutTemplateMapInfo = nullptr);
 
 	/**
 	 * Creates and show a window with an SNewLevelDialog
@@ -43,9 +46,10 @@ public:
 	 * @param	OutTemplateName	- (out) The package name of the template map selected by the user. Empty if blank map selected.
 	 * @param	bShowPartitionedTemplates - If true, we allow partitioned world templates to be shown in the dialog.
 	 * @param	bOutPartitionedWorld - (out) If true and OutTemplateName is empty, the empty map should be partitioned.
+	 * @param	OutTemplateMapInfo - (out) If not null, will be set to the template map selected by the user.
 	 * @return	true if the user selected a valid item, false if the user canceled
 	 */
-	virtual bool CreateAndShowTemplateDialog(const TSharedPtr<const SWidget> ParentWidget, const FText& Title, const TArray<FTemplateMapInfo>& Templates, FString& OutTemplateMapPackageName, bool bShowPartitionedTemplates, bool& bOutPartitionedWorld);
+	virtual bool CreateAndShowTemplateDialog(const TSharedPtr<const SWidget> ParentWidget, const FText& Title, const TArray<FTemplateMapInfo>& Templates, FString& OutTemplateMapPackageName, bool bShowPartitionedTemplates, bool& bOutPartitionedWorld, FTemplateMapInfo* OutTemplateMapInfo = nullptr);
 
 	/** New Level Dialog app identifier string */
 	static const FName NewLevelDialogAppIdentifier;
