Fix new level dialog issues with sort and invalidating cooked textures

diff --git a/Engine/Source/Editor/NewLevelDialog/Private/NewLevelDialogModule.cpp b/Engine/Source/Editor/NewLevelDialog/Private/NewLevelDialogModule.cpp
--- a/Engine/Source/Editor/NewLevelDialog/Private/NewLevelDialogModule.cpp
+++ b/Engine/Source/Editor/NewLevelDialog/Private/NewLevelDialogModule.cpp
@@ -50,6 +50,7 @@ struct FNewLevelTemplateItem
 	FString Category;
 	TUniquePtr<FSlateBrush> ThumbnailBrush;
 	UTexture2D* ThumbnailAsset = nullptr;
+	int32 OriginalIndex = INDEX_NONE;
 
 	enum NewLevelType
 	{
@@ -355,10 +356,12 @@ private:
 		TArray<TSharedPtr<FNewLevelTemplateItem>> TemplateItems;
 
 		// Build a list of items - one for each template
+		int32 CurrentIndex = 0;
 		for (const FTemplateMapInfo& TemplateMapInfo : TemplateMapInfos)
 		{
 			if (!bShowPartitionedTemplates && ULevel::GetIsLevelPartitionedFromPackage(FName(*TemplateMapInfo.Map.ToString())))
 			{
+				CurrentIndex++;
 				continue;
 			}
 
@@ -367,34 +370,34 @@ private:
 			Item->Type = FNewLevelTemplateItem::NewLevelType::Template;
 			Item->Name = TemplateMapInfo.DisplayName;
 			Item->Category = TemplateMapInfo.Category;
+			Item->OriginalIndex = CurrentIndex++;
 
-			UTexture2D* ThumbnailTexture = nullptr;
+			TSoftObjectPtr<UTexture2D> ThumbnailSoftObjectPtr;
 			if (TemplateMapInfo.Thumbnail.IsValid())
 			{
-				FSoftObjectPath SoftObjectPath(TemplateMapInfo.Thumbnail);
-				ThumbnailTexture = Cast<UTexture2D>(SoftObjectPath.ResolveObject());
-				if (!ThumbnailTexture)
-				{
-					ThumbnailTexture = Cast<UTexture2D>(SoftObjectPath.TryLoad());
-					if (ThumbnailTexture)
-					{
-						// Newly loaded texture requires async work to complete before being rendered in modal dialog
-						ThumbnailTexture->FinishCachePlatformData();
-						ThumbnailTexture->UpdateResource();
-					}
-				}
+				ThumbnailSoftObjectPtr = TemplateMapInfo.Thumbnail;
 			}
 			else if (!TemplateMapInfo.ThumbnailTexture.IsNull())
 			{
-				ThumbnailTexture = TemplateMapInfo.ThumbnailTexture.Get();
+				ThumbnailSoftObjectPtr = TemplateMapInfo.ThumbnailTexture;
+			}
+
+			UTexture2D* ThumbnailTexture = nullptr;
+			if (!ThumbnailSoftObjectPtr.IsNull())
+			{
+				ThumbnailTexture = ThumbnailSoftObjectPtr.Get();
 				if (!ThumbnailTexture)
 				{
-					ThumbnailTexture = TemplateMapInfo.ThumbnailTexture.LoadSynchronous();
+					ThumbnailTexture = ThumbnailSoftObjectPtr.LoadSynchronous();
 					if (ThumbnailTexture)
 					{
-						// Newly loaded texture requires async work to complete before being rendered in modal dialog
-						ThumbnailTexture->FinishCachePlatformData();
-						ThumbnailTexture->UpdateResource();
+						// Avoid calling UpdateResource on cooked texture as doing so will destroy the texture's data
+						if (!ThumbnailTexture->GetPackage()->HasAllPackagesFlags(PKG_FilterEditorOnly))
+						{
+							// Newly loaded texture requires async work to complete before being rendered in modal dialog
+							ThumbnailTexture->FinishCachePlatformData();
+							ThumbnailTexture->UpdateResource();
+						}
 					}
 				}
 			}
@@ -464,6 +467,11 @@ private:
 		{
 			if (LHS->Category == RHS->Category)
 			{
+				if (LHS->Type == RHS->Type)
+				{
+					return LHS->OriginalIndex < RHS->OriginalIndex;
+				}
+
 				return LHS->Type > RHS->Type;
 			}
 			else if (LHS->Category == OpenWorldCategory)
