Add ability to hide empty level template

diff --git a/Engine/Source/Editor/NewLevelDialog/Private/NewLevelDialogModule.cpp b/Engine/Source/Editor/NewLevelDialog/Private/NewLevelDialogModule.cpp
--- a/Engine/Source/Editor/NewLevelDialog/Private/NewLevelDialogModule.cpp
+++ b/Engine/Source/Editor/NewLevelDialog/Private/NewLevelDialogModule.cpp
@@ -225,13 +225,14 @@ public:
 
 		SLATE_ATTRIBUTE(TArray<FTemplateMapInfo>, Templates)
 		SLATE_ATTRIBUTE(bool, bShowPartitionedTemplates)
+		SLATE_ATTRIBUTE(FNewLevelDialogModule::EShowEmptyTemplate, ShowEmptyTemplate)
 	SLATE_END_ARGS()
 
 	void Construct(const FArguments& InArgs)
 	{
 		ParentWindowPtr = InArgs._ParentWindow.Get();
 
-		TemplateItemsList = MakeTemplateItems(InArgs._bShowPartitionedTemplates.Get(), InArgs._Templates.Get());
+		TemplateItemsList = MakeTemplateItems(InArgs._bShowPartitionedTemplates.Get(), InArgs._ShowEmptyTemplate.Get(), InArgs._Templates.Get());
 		
 		TemplateListView = SNew(STileView<TSharedPtr<FNewLevelTemplateItem>>)
 			.ListItemsSource(&TemplateItemsList)
@@ -351,7 +352,7 @@ public:
 	//~ End FGCObject Interface.
 
 private:
-	static TArray<TSharedPtr<FNewLevelTemplateItem>> MakeTemplateItems(bool bShowPartitionedTemplates, const TArray<FTemplateMapInfo>& TemplateMapInfos)
+	static TArray<TSharedPtr<FNewLevelTemplateItem>> MakeTemplateItems(bool bShowPartitionedTemplates, FNewLevelDialogModule::EShowEmptyTemplate ShowEmptyTemplate, const TArray<FTemplateMapInfo>& TemplateMapInfos)
 	{
 		TArray<TSharedPtr<FNewLevelTemplateItem>> TemplateItems;
 
@@ -437,14 +438,17 @@ private:
 		}
 
 		// Add an extra item for creating a new, blank level
-		TSharedPtr<FNewLevelTemplateItem> NewItem = MakeShareable(new FNewLevelTemplateItem());
-		NewItem->Type = FNewLevelTemplateItem::NewLevelType::Empty;
-		NewItem->Name = LOCTEXT("NewLevelItemLabel", "Empty Level");
-		NewItem->ThumbnailBrush = MakeUnique<FSlateBrush>(*FAppStyle::GetBrush("NewLevelDialog.Blank"));
-		NewItem->ThumbnailBrush->OutlineSettings.CornerRadii = FVector4(4, 4, 0, 0);
-		NewItem->ThumbnailBrush->OutlineSettings.RoundingType = ESlateBrushRoundingType::FixedRadius;
-		NewItem->ThumbnailBrush->DrawAs = ESlateBrushDrawType::RoundedBox;
-		TemplateItems.Add(NewItem);
+		if (ShowEmptyTemplate == FNewLevelDialogModule::EShowEmptyTemplate::Show)
+		{
+			TSharedPtr<FNewLevelTemplateItem> NewItem = MakeShareable(new FNewLevelTemplateItem());
+			NewItem->Type = FNewLevelTemplateItem::NewLevelType::Empty;
+			NewItem->Name = LOCTEXT("NewLevelItemLabel", "Empty Level");
+			NewItem->ThumbnailBrush = MakeUnique<FSlateBrush>(*FAppStyle::GetBrush("NewLevelDialog.Blank"));
+			NewItem->ThumbnailBrush->OutlineSettings.CornerRadii = FVector4(4, 4, 0, 0);
+			NewItem->ThumbnailBrush->OutlineSettings.RoundingType = ESlateBrushRoundingType::FixedRadius;
+			NewItem->ThumbnailBrush->DrawAs = ESlateBrushDrawType::RoundedBox;
+			TemplateItems.Add(NewItem);
+		}
 
 		const FString& OpenWorldCategory = TEXT("OpenWorld");
 
@@ -562,13 +566,13 @@ void FNewLevelDialogModule::ShutdownModule()
 {
 }
 
-bool FNewLevelDialogModule::CreateAndShowNewLevelDialog( const TSharedPtr<const SWidget> ParentWidget, FString& OutTemplateMapPackageName, bool bShowPartitionedTemplates, bool& bOutIsPartitionedWorld, FTemplateMapInfo* OutTemplateMapInfo)
+bool FNewLevelDialogModule::CreateAndShowNewLevelDialog( const TSharedPtr<const SWidget> ParentWidget, FString& OutTemplateMapPackageName, bool bShowPartitionedTemplates, bool& bOutIsPartitionedWorld, FTemplateMapInfo* OutTemplateMapInfo, EShowEmptyTemplate ShowEmptyTemplate)
 {
 	TArray<FTemplateMapInfo> EmptyTemplates;
-	return CreateAndShowTemplateDialog(ParentWidget, LOCTEXT("WindowHeader", "New Level"), GUnrealEd ? GUnrealEd->GetTemplateMapInfos() : EmptyTemplates, OutTemplateMapPackageName, bShowPartitionedTemplates, bOutIsPartitionedWorld, OutTemplateMapInfo);
+	return CreateAndShowTemplateDialog(ParentWidget, LOCTEXT("WindowHeader", "New Level"), GUnrealEd ? GUnrealEd->GetTemplateMapInfos() : EmptyTemplates, OutTemplateMapPackageName, bShowPartitionedTemplates, bOutIsPartitionedWorld, OutTemplateMapInfo, ShowEmptyTemplate);
 }
 
-bool FNewLevelDialogModule::CreateAndShowTemplateDialog( const TSharedPtr<const SWidget> ParentWidget, const FText& Title, const TArray<FTemplateMapInfo>& Templates, FString& OutTemplateMapPackageName, bool bShowPartitionedTemplates, bool& bOutIsPartitionedWorld, FTemplateMapInfo* OutTemplateMapInfo)
+bool FNewLevelDialogModule::CreateAndShowTemplateDialog( const TSharedPtr<const SWidget> ParentWidget, const FText& Title, const TArray<FTemplateMapInfo>& Templates, FString& OutTemplateMapPackageName, bool bShowPartitionedTemplates, bool& bOutIsPartitionedWorld, FTemplateMapInfo* OutTemplateMapInfo, EShowEmptyTemplate ShowEmptyTemplate)
 {
 	// Open larger window if there are enough templates
 	FVector2D WindowClientSize(NewLevelDialogDefs::DefaultWindowWidth, NewLevelDialogDefs::DefaultWindowHeight);
@@ -591,7 +595,8 @@ bool FNewLevelDialogModule::CreateAndShowTemplateDialog( const TSharedPtr<const
 		SNew(SNewLevelDialog)
 		.ParentWindow(NewLevelWindow)
 		.Templates(Templates)
-		.bShowPartitionedTemplates(bShowPartitionedTemplates);
+		.bShowPartitionedTemplates(bShowPartitionedTemplates)
+		.ShowEmptyTemplate(ShowEmptyTemplate);
 
 	NewLevelWindow->SetContent(NewLevelDialog);
 
diff --git a/Engine/Source/Editor/NewLevelDialog/Public/NewLevelDialogModule.h b/Engine/Source/Editor/NewLevelDialog/Public/NewLevelDialogModule.h
--- a/Engine/Source/Editor/NewLevelDialog/Public/NewLevelDialogModule.h
+++ b/Engine/Source/Editor/NewLevelDialog/Public/NewLevelDialogModule.h
@@ -15,6 +15,16 @@ struct FTemplateMapInfo;
 class FNewLevelDialogModule : public IModuleInterface
 {
 public:
+
+	/**
+	 * Control if 'Empty Level' template should appear or not
+	 */
+	enum class EShowEmptyTemplate : uint8
+	{
+		Show,
+		Hide
+	};
+
 	/**
 	 * Called right after the plugin DLL has been loaded and the plugin object has been created
 	 */
@@ -33,9 +43,10 @@ public:
 	 * @param	bShowPartitionedTemplates - If true, we allow partitioned world templates to be shown in the dialog.
 	 * @param	bOutPartitionedWorld - (out) If true and OutTemplateName is empty, the empty map should be partitioned.
 	 * @param	OutTemplateMapInfo - (out) If not null, will be set to the template map selected by the user.
+	 * @param	ShowEmptyTemplate - Control if 'Empty Level' template should appear or not
 	 * @return	true if the user selected a valid item, false if the user canceled
 	 */
-	virtual bool CreateAndShowNewLevelDialog( const TSharedPtr<const SWidget> ParentWidget, FString& OutTemplateMapPackageName, bool bShowPartitionedTemplates, bool& bOutPartitionedWorld, FTemplateMapInfo* OutTemplateMapInfo = nullptr);
+	virtual bool CreateAndShowNewLevelDialog( const TSharedPtr<const SWidget> ParentWidget, FString& OutTemplateMapPackageName, bool bShowPartitionedTemplates, bool& bOutPartitionedWorld, FTemplateMapInfo* OutTemplateMapInfo = nullptr, EShowEmptyTemplate ShowEmptyTemplate = EShowEmptyTemplate::Show);
 
 	/**
 	 * Creates and show a window with an SNewLevelDialog
@@ -47,9 +58,10 @@ public:
 	 * @param	bShowPartitionedTemplates - If true, we allow partitioned world templates to be shown in the dialog.
 	 * @param	bOutPartitionedWorld - (out) If true and OutTemplateName is empty, the empty map should be partitioned.
 	 * @param	OutTemplateMapInfo - (out) If not null, will be set to the template map selected by the user.
+	 * @param	ShowEmptyTemplate - Control if 'Empty Level' template should appear or not
 	 * @return	true if the user selected a valid item, false if the user canceled
 	 */
-	virtual bool CreateAndShowTemplateDialog(const TSharedPtr<const SWidget> ParentWidget, const FText& Title, const TArray<FTemplateMapInfo>& Templates, FString& OutTemplateMapPackageName, bool bShowPartitionedTemplates, bool& bOutPartitionedWorld, FTemplateMapInfo* OutTemplateMapInfo = nullptr);
+	virtual bool CreateAndShowTemplateDialog(const TSharedPtr<const SWidget> ParentWidget, const FText& Title, const TArray<FTemplateMapInfo>& Templates, FString& OutTemplateMapPackageName, bool bShowPartitionedTemplates, bool& bOutPartitionedWorld, FTemplateMapInfo* OutTemplateMapInfo = nullptr, EShowEmptyTemplate ShowEmptyTemplate = EShowEmptyTemplate::Show);
 
 	/** New Level Dialog app identifier string */
 	static const FName NewLevelDialogAppIdentifier;
